Fix for crash on exit in VPCustomUIHandler

diff --git a/Engine/Plugins/Experimental/VirtualProductionUtilities/Source/VPUtilitiesEditor/Private/VPUtilitiesEditorModule.cpp b/Engine/Plugins/Experimental/VirtualProductionUtilities/Source/VPUtilitiesEditor/Private/VPUtilitiesEditorModule.cpp
--- a/Engine/Plugins/Experimental/VirtualProductionUtilities/Source/VPUtilitiesEditor/Private/VPUtilitiesEditorModule.cpp
+++ b/Engine/Plugins/Experimental/VirtualProductionUtilities/Source/VPUtilitiesEditor/Private/VPUtilitiesEditorModule.cpp
@@ -55,7 +55,11 @@ public:
 		UnregisterSettings();
 		SGenlockProviderTab::UnregisterNomadTabSpawner();
 
-		CustomUIHandler->Uninit();
+		if (UObjectInitialized())
+		{
+			CustomUIHandler->Uninit();
+		}
+
 		CustomUIHandler.Reset();
 
 		FVPUtilitiesEditorStyle::Unregister();
