Move VirtualProduction asset to a editor folder.
Delay the load of VirtualProductionWidget to prevent issue with cooking and nativization.

diff --git a/Engine/Config/BaseEditor.ini b/Engine/Config/BaseEditor.ini
--- a/Engine/Config/BaseEditor.ini
+++ b/Engine/Config/BaseEditor.ini
@@ -450,8 +450,7 @@ bNativizeAnimBPOnlyWhenNonReducibleFuncitons=false
 +ExcludedAssets=/Engine/Tutorial/SubEditors/TutorialAssets/Character/TutorialCharacter
 +ExcludedAssets=/Engine/Tutorial/SubEditors/TutorialAssets/Character/TutorialTPP_AnimBlueprint
 
-; Currently excluded due to an editor-only asset dependency
-+ExcludedAssets=/VirtualProductionUtilities/UserWidgets/BookmarkWidget
+; +ExcludedFolderPaths=...
 
 ; To add a module dependency to the nativized module in both client and server use:
 ; +AdditionalPublicDependencyModuleNames=
diff --git a/Engine/Plugins/Experimental/VirtualProductionUtilities/Source/VPUtilitiesEditor/Private/VPCustomUIHandler.cpp b/Engine/Plugins/Experimental/VirtualProductionUtilities/Source/VPUtilitiesEditor/Private/VPCustomUIHandler.cpp
--- a/Engine/Plugins/Experimental/VirtualProductionUtilities/Source/VPUtilitiesEditor/Private/VPCustomUIHandler.cpp
+++ b/Engine/Plugins/Experimental/VirtualProductionUtilities/Source/VPUtilitiesEditor/Private/VPCustomUIHandler.cpp
@@ -12,6 +12,7 @@
 #include "UI/VREditorUISystem.h"
 #include "UObject/ConstructorHelpers.h"
 #include "VREditorStyle.h"
+#include "VPUtilitiesEditorModule.h"
 #include "WidgetBlueprint.h"
 #include "Widgets/SWidget.h"
 
@@ -21,16 +22,6 @@
 static const FName VirtualProductionToolsLabel = TEXT("VirtualProductionTools");
 
 
-UVPCustomUIHandler::UVPCustomUIHandler(const FObjectInitializer& ObjectInitializer)
-	: Super(ObjectInitializer)
-{
-	static ConstructorHelpers::FObjectFinder<UWidgetBlueprint> DefaultWidget(TEXT("/VirtualProductionUtilities/VirtualProductionWidget"));
-	if (DefaultWidget.Object)
-	{
-		VirtualProductionWidget = DefaultWidget.Object->GeneratedClass;
-	}
-}
-
 void UVPCustomUIHandler::Init()
 {
 	VRRadialMenuWindowsExtension = IVREditorModule::Get().GetRadialMenuExtender()->AddMenuExtension(
@@ -52,6 +43,20 @@ void UVPCustomUIHandler::Uninit()
 
 void UVPCustomUIHandler::FillVRRadialMenuWindows(FMenuBuilder& MenuBuilder)
 {
+	if (VirtualProductionWidget.Get() == nullptr)
+	{
+		UWidgetBlueprint* WidgetBlueprint = LoadObject<UWidgetBlueprint>(nullptr, TEXT("/VirtualProductionUtilities/Editor/VirtualProductionWidget.VirtualProductionWidget"));
+		if (WidgetBlueprint)
+		{
+			VirtualProductionWidget = WidgetBlueprint->GeneratedClass;
+		}
+	}
+
+	if (VirtualProductionWidget.Get() == nullptr)
+	{
+		UE_LOG(LogVPUtilitiesEditor, Warning, TEXT("/VirtualProductionUtilities/Editor/VirtualProductionWidget could not be found."));
+	}
+
 	MenuBuilder.AddMenuEntry(
 		LOCTEXT("VirtualProductionTools", "Virtual Production"),
 		FText(),
diff --git a/Engine/Plugins/Experimental/VirtualProductionUtilities/Source/VPUtilitiesEditor/Private/VPCustomUIHandler.h b/Engine/Plugins/Experimental/VirtualProductionUtilities/Source/VPUtilitiesEditor/Private/VPCustomUIHandler.h
--- a/Engine/Plugins/Experimental/VirtualProductionUtilities/Source/VPUtilitiesEditor/Private/VPCustomUIHandler.h
+++ b/Engine/Plugins/Experimental/VirtualProductionUtilities/Source/VPUtilitiesEditor/Private/VPCustomUIHandler.h
@@ -16,8 +16,6 @@ class UVPCustomUIHandler : public UObject
 {
 	GENERATED_BODY()
 
-	UVPCustomUIHandler(const FObjectInitializer& ObjectInitializer);
-
 public:
 	void Init();
 	void Uninit();
