Move Virtual Production Utilities plugin from "Other" to "Virtual Production"

diff --git a/Engine/Plugins/Experimental/VirtualProductionUtilities/VirtualProductionUtilities.uplugin b/Engine/Plugins/Experimental/VirtualProductionUtilities/VirtualProductionUtilities.uplugin
--- a/Engine/Plugins/Experimental/VirtualProductionUtilities/VirtualProductionUtilities.uplugin
+++ b/Engine/Plugins/Experimental/VirtualProductionUtilities/VirtualProductionUtilities.uplugin
@@ -4,7 +4,7 @@
 	"VersionName": "1.0",
 	"FriendlyName": "Virtual Production Utilities",
 	"Description": "",
-	"Category": "Other",
+	"Category": "Virtual Production",
 	"CreatedBy": "",
 	"CreatedByURL": "",
 	"DocsURL": "",
