Integrating fix from private branch

- Changes VPUtilitiesEditor module LoadingPhase to PostEngineInit to ensure that the Blutility module is loaded when VPScoutingSubsystem is constructed. Otherwise, two BPs are not/cannot be correctly identified as Editor Utility BPs and fail to compile during UED startup.
- Changed three modules in VPUtilitiesEditor build file to be public dependencies.

diff --git a/Engine/Plugins/Experimental/VirtualProductionUtilities/Source/VPUtilitiesEditor/VPUtilitiesEditor.Build.cs b/Engine/Plugins/Experimental/VirtualProductionUtilities/Source/VPUtilitiesEditor/VPUtilitiesEditor.Build.cs
--- a/Engine/Plugins/Experimental/VirtualProductionUtilities/Source/VPUtilitiesEditor/VPUtilitiesEditor.Build.cs
+++ b/Engine/Plugins/Experimental/VirtualProductionUtilities/Source/VPUtilitiesEditor/VPUtilitiesEditor.Build.cs
@@ -11,9 +11,12 @@ public class VPUtilitiesEditor : ModuleRules
 		PublicDependencyModuleNames.AddRange(
 			new string[]
 			{
+                "Blutility",
 				"Core",
 				"CoreUObject",
+				"EditorSubsystem",
 				"VPUtilities",                
+				"VREditor",
             }
 		);
 
@@ -32,10 +35,7 @@ public class VPUtilitiesEditor : ModuleRules
 				"UMGEditor",
 				"UnrealEd",
 				"VPBookmark",
-				"VREditor",
 				"WorkspaceMenuStructure",
-				"EditorSubsystem",
-                "Blutility",
                 "CinematicCamera",
             }
 		);
diff --git a/Engine/Plugins/Experimental/VirtualProductionUtilities/VirtualProductionUtilities.uplugin b/Engine/Plugins/Experimental/VirtualProductionUtilities/VirtualProductionUtilities.uplugin
--- a/Engine/Plugins/Experimental/VirtualProductionUtilities/VirtualProductionUtilities.uplugin
+++ b/Engine/Plugins/Experimental/VirtualProductionUtilities/VirtualProductionUtilities.uplugin
@@ -22,7 +22,7 @@
 		{
 			"Name": "VPUtilitiesEditor",
 			"Type": "Editor",
-			"LoadingPhase": "PostDefault"
+			"LoadingPhase": "PostEngineInit"
 		},
 		{
 			"Name": "VPBookmark",
