Add virtual production roles

diff --git a/Engine/Plugins/Experimental/VirtualProductionUtilities/Source/VPUtilities/Private/VPBlueprintLibrary.cpp b/Engine/Plugins/Experimental/VirtualProductionUtilities/Source/VPUtilities/Private/VPBlueprintLibrary.cpp
--- a/Engine/Plugins/Experimental/VirtualProductionUtilities/Source/VPUtilities/Private/VPBlueprintLibrary.cpp
+++ b/Engine/Plugins/Experimental/VirtualProductionUtilities/Source/VPUtilities/Private/VPBlueprintLibrary.cpp
@@ -4,6 +4,7 @@
 
 #include "Engine/World.h"
 #include "HAL/IConsoleManager.h"
+#include "VPSettings.h"
 #include "VPUtilitiesModule.h"
 
 #if WITH_EDITOR
@@ -96,6 +97,12 @@ bool UVPBlueprintLibrary::JumpToBookmarkInLevelEditor(const UVPBookmark* Bookmar
 }
 
 
+FGameplayTagContainer UVPBlueprintLibrary::GetVirtualProductionRole()
+{
+	return GetDefault<UVPSettings>()->GetRoles();
+}
+
+
 FTransform UVPBlueprintLibrary::GetEditorViewportTransform()
 {
 #if WITH_EDITOR
diff --git a/Engine/Plugins/Experimental/VirtualProductionUtilities/Source/VPUtilities/Private/VPSettings.cpp b/Engine/Plugins/Experimental/VirtualProductionUtilities/Source/VPUtilities/Private/VPSettings.cpp
new file mode 100644
--- /dev/null
+++ b/Engine/Plugins/Experimental/VirtualProductionUtilities/Source/VPUtilities/Private/VPSettings.cpp
@@ -0,0 +1,44 @@
+// Copyright 1998-2019 Epic Games, Inc. All Rights Reserved.
+
+#include "VPSettings.h"
+
+#include "Misc/CommandLine.h"
+#include "VPUtilitiesModule.h"
+
+UVPSettings::UVPSettings()
+{
+	FString ComandlineRoles;
+	bIsCommandLineRolesValid = FParse::Value(FCommandLine::Get(), TEXT("-VPRole="), ComandlineRoles);
+
+	if (bIsCommandLineRolesValid)
+	{
+		TArray<FString> RoleList;
+		ComandlineRoles.ParseIntoArray(RoleList, TEXT("|"), true);
+
+		for (const FString& Role : RoleList)
+		{
+			FGameplayTag Tag = FGameplayTag::RequestGameplayTag(*Role, false);
+			if (Tag.IsValid())
+			{
+				CommandLineRoles.AddTag(Tag);
+			}
+			else
+			{
+				UE_LOG(LogVPUtilities, Fatal, TEXT("Role %s doesn't exist."), *Role);
+			}
+		}
+	}
+}
+
+
+#if WITH_EDITOR
+void UVPSettings::PostEditChangeChainProperty(struct FPropertyChangedChainEvent& PropertyChangedEvent)
+{
+	if (PropertyChangedEvent.Property->GetFName() == GET_MEMBER_NAME_CHECKED(ThisClass, Roles))
+	{
+		OnRolesChanged.Broadcast();
+	}
+
+	Super::PostEditChangeChainProperty(PropertyChangedEvent);
+}
+#endif
diff --git a/Engine/Plugins/Experimental/VirtualProductionUtilities/Source/VPUtilities/Public/VPBlueprintLibrary.h b/Engine/Plugins/Experimental/VirtualProductionUtilities/Source/VPUtilities/Public/VPBlueprintLibrary.h
--- a/Engine/Plugins/Experimental/VirtualProductionUtilities/Source/VPUtilities/Public/VPBlueprintLibrary.h
+++ b/Engine/Plugins/Experimental/VirtualProductionUtilities/Source/VPUtilities/Public/VPBlueprintLibrary.h
@@ -3,6 +3,7 @@
 #pragma once
 
 #include "Kismet/BlueprintFunctionLibrary.h"
+#include "GameplayTagContainer.h"
 #include "VPBookmarkContext.h"
 #include "VPViewportTickableActorBase.h"
 #include "VPBlueprintLibrary.generated.h"
@@ -32,6 +33,10 @@ public:
 	UFUNCTION(BlueprintCallable, Category = "Virtual Production")
 	static bool JumpToBookmarkInLevelEditor(const UVPBookmark* Bookmark);
 
+	/** The machine role(s) in a virtual production context. */
+	UFUNCTION(BlueprintCallable, Category = "Virtual Production")
+	static FGameplayTagContainer GetVirtualProductionRole();
+
 	/** Get the location of the 2D viewport camera */
 	UFUNCTION(BlueprintCallable, Category = "Virtual Production")
 	static FTransform GetEditorViewportTransform();
diff --git a/Engine/Plugins/Experimental/VirtualProductionUtilities/Source/VPUtilities/Public/VPSettings.h b/Engine/Plugins/Experimental/VirtualProductionUtilities/Source/VPUtilities/Public/VPSettings.h
new file mode 100644
--- /dev/null
+++ b/Engine/Plugins/Experimental/VirtualProductionUtilities/Source/VPUtilities/Public/VPSettings.h
@@ -0,0 +1,48 @@
+// Copyright 1998-2019 Epic Games, Inc. All Rights Reserved.
+
+#pragma once
+
+#include "CoreMinimal.h"
+#include "GameplayTagContainer.h"
+#include "UObject/Object.h"
+#include "VPSettings.generated.h"
+
+/**
+ * 
+ */
+UCLASS(config=Game)
+class VPUTILITIES_API UVPSettings : public UObject
+{
+	GENERATED_BODY()
+
+public:
+	UVPSettings();
+
+protected:
+	/**
+	 * The machine role(s) in a virtual production context.
+	 * @note The role may be override via the command line, "-VPRole=[Role.SubRole1|Role.SubRole2]"
+	 */
+	UPROPERTY(config, EditAnywhere, Category="Virtual Production")
+	FGameplayTagContainer Roles;
+
+	/**
+	 * The machine role(s) in a virtual production context read from the command line.
+	 * ie. "-VPRole=[Role.SubRole1|Role.SubRole2]"
+	 */
+	bool bIsCommandLineRolesValid;
+	UPROPERTY(transient)
+	FGameplayTagContainer CommandLineRoles;
+
+public:
+	const FGameplayTagContainer& GetRoles() const { return bIsCommandLineRolesValid ? CommandLineRoles : Roles; }
+
+public:
+#if WITH_EDITOR
+	//~ UObject interface
+	virtual void PostEditChangeChainProperty(struct FPropertyChangedChainEvent& PropertyChangedEvent) override;
+	//~ End UObject interface
+
+	FSimpleMulticastDelegate OnRolesChanged;
+#endif //WITH_EDITOR
+};
diff --git a/Engine/Plugins/Experimental/VirtualProductionUtilities/Source/VPUtilities/VPUtilities.Build.cs b/Engine/Plugins/Experimental/VirtualProductionUtilities/Source/VPUtilities/VPUtilities.Build.cs
--- a/Engine/Plugins/Experimental/VirtualProductionUtilities/Source/VPUtilities/VPUtilities.Build.cs
+++ b/Engine/Plugins/Experimental/VirtualProductionUtilities/Source/VPUtilities/VPUtilities.Build.cs
@@ -11,10 +11,11 @@ public class VPUtilities : ModuleRules
 		PublicDependencyModuleNames.AddRange(
 			new string[]
 			{
+				"CinematicCamera",
 				"Core",
 				"CoreUObject",
 				"Engine",
-				"CinematicCamera",
+				"GameplayTags",
 				"VPBookmark",
 			}
 		);
diff --git a/Engine/Plugins/Experimental/VirtualProductionUtilities/Source/VPUtilitiesEditor/Private/VPUtilitiesEditorModule.cpp b/Engine/Plugins/Experimental/VirtualProductionUtilities/Source/VPUtilitiesEditor/Private/VPUtilitiesEditorModule.cpp
--- a/Engine/Plugins/Experimental/VirtualProductionUtilities/Source/VPUtilitiesEditor/Private/VPUtilitiesEditorModule.cpp
+++ b/Engine/Plugins/Experimental/VirtualProductionUtilities/Source/VPUtilitiesEditor/Private/VPUtilitiesEditorModule.cpp
@@ -3,11 +3,14 @@
 #include "VPUtilitiesEditorModule.h"
 
 #include "Framework/Docking/WorkspaceItem.h"
+#include "ISettingsModule.h"
+#include "ISettingsSection.h"
 #include "Modules/ModuleInterface.h"
 #include "Modules/ModuleManager.h"
 #include "SGenlockProviderTab.h"
 #include "Textures/SlateIcon.h"
 #include "VPCustomUIHandler.h"
+#include "VPSettings.h"
 #include "VPUtilitiesEditorStyle.h"
 #include "UObject/StrongObjectPtr.h"
 #include "WorkspaceMenuStructure.h"
@@ -28,6 +31,7 @@ public:
 	virtual void StartupModule() override
 	{
 		FVPUtilitiesEditorStyle::Register();
+
 		CustomUIHandler.Reset(NewObject<UVPCustomUIHandler>());
 		CustomUIHandler->Init();
 
@@ -40,17 +44,45 @@ public:
 
 			SGenlockProviderTab::RegisterNomadTabSpawner(MediaBrowserGroup);
 		}
+
+		RegisterSettings();
 	}
 
 
 	virtual void ShutdownModule() override
 	{
+
+		UnregisterSettings();
 		SGenlockProviderTab::UnregisterNomadTabSpawner();
 
 		CustomUIHandler->Uninit();
 		CustomUIHandler.Reset();
+
 		FVPUtilitiesEditorStyle::Unregister();
 	}
+
+
+	void RegisterSettings()
+	{
+		ISettingsModule* SettingsModule = FModuleManager::GetModulePtr<ISettingsModule>("Settings");
+		if (SettingsModule != nullptr)
+		{
+			ISettingsSectionPtr SettingsSection = SettingsModule->RegisterSettings("Project", "Plugins", "VirtualProduction",
+				LOCTEXT("VirtualProductionSettingsName", "Virtual Production"),
+				LOCTEXT("VirtualProductionSettingsDescription", "Configure the Virtual Production settings."),
+				GetMutableDefault<UVPSettings>());
+		}
+	}
+
+
+	void UnregisterSettings()
+	{
+		ISettingsModule* SettingsModule = FModuleManager::GetModulePtr<ISettingsModule>("Settings");
+		if (SettingsModule != nullptr)
+		{
+			SettingsModule->UnregisterSettings("Project", "Plugins", "VirtualProduction");
+		}
+	}
 };
 
 
diff --git a/Engine/Plugins/Experimental/VirtualProductionUtilities/Source/VPUtilitiesEditor/VPUtilitiesEditor.Build.cs b/Engine/Plugins/Experimental/VirtualProductionUtilities/Source/VPUtilitiesEditor/VPUtilitiesEditor.Build.cs
--- a/Engine/Plugins/Experimental/VirtualProductionUtilities/Source/VPUtilitiesEditor/VPUtilitiesEditor.Build.cs
+++ b/Engine/Plugins/Experimental/VirtualProductionUtilities/Source/VPUtilitiesEditor/VPUtilitiesEditor.Build.cs
@@ -23,6 +23,7 @@ public class VPUtilitiesEditor : ModuleRules
 				"EditorStyle",
 				"Engine",
 				"LevelEditor",
+				"Settings",
 				"Slate",
 				"SlateCore",
 				"TimeManagement",
