VirtualProductionUtilities

- Bug fix for Get Target Camera

diff --git a/Engine/Plugins/Experimental/VirtualProductionUtilities/Source/VPUtilities/Private/VPRootActor.cpp b/Engine/Plugins/Experimental/VirtualProductionUtilities/Source/VPUtilities/Private/VPRootActor.cpp
--- a/Engine/Plugins/Experimental/VirtualProductionUtilities/Source/VPUtilities/Private/VPRootActor.cpp
+++ b/Engine/Plugins/Experimental/VirtualProductionUtilities/Source/VPUtilities/Private/VPRootActor.cpp
@@ -125,7 +125,7 @@ namespace VPRootActor
 
 ACineCameraActor* AVPRootActor::GetCineCameraActor_Implementation() const
 {
-	if (CinematicCamera == nullptr)
+	if (CinematicCamera != nullptr)
 	{
 		return CinematicCamera;
 	}
