- Fix for VRTools multi user functionnalities

diff --git a/Engine/Plugins/Experimental/VirtualProductionUtilities/Source/VPUtilitiesEditor/Private/VPUtilitiesEditorBlueprintLibrary.cpp b/Engine/Plugins/Experimental/VirtualProductionUtilities/Source/VPUtilitiesEditor/Private/VPUtilitiesEditorBlueprintLibrary.cpp
--- a/Engine/Plugins/Experimental/VirtualProductionUtilities/Source/VPUtilitiesEditor/Private/VPUtilitiesEditorBlueprintLibrary.cpp
+++ b/Engine/Plugins/Experimental/VirtualProductionUtilities/Source/VPUtilitiesEditor/Private/VPUtilitiesEditorBlueprintLibrary.cpp
@@ -25,3 +25,24 @@ AVPEditorTickableActorBase* UVPUtilitiesEditorBlueprintLibrary::SpawnVPEditorTic
 	AVPEditorTickableActorBase* NewActor = World->SpawnActor<AVPEditorTickableActorBase>(ActorClass.Get(), Location, Rotation, SpawnParams);
 	return NewActor;
 }
+
+AVPTransientEditorTickableActorBase* UVPUtilitiesEditorBlueprintLibrary::SpawnVPTransientEditorTickableActor(UObject* ContextObject, const TSubclassOf<AVPTransientEditorTickableActorBase> ActorClass, const FVector Location, const FRotator Rotation)
+{
+	if (ActorClass.Get() == nullptr)
+	{
+		UE_LOG(LogVPUtilitiesEditor, Warning, TEXT("VPUtilitiesEditorBlueprintLibrary::SpawnVPTransientEditorTickableActor - The ActorClass is invalid"));
+		return nullptr;
+	}
+
+	UWorld* World = ContextObject ? ContextObject->GetWorld() : nullptr;
+	if (World == nullptr)
+	{
+		UE_LOG(LogVPUtilitiesEditor, Warning, TEXT("VPUtilitiesEditorBlueprintLibrary::SpawnVPTransientEditorTickableActor - The ContextObject is invalid."));
+		return nullptr;
+	}
+
+	FActorSpawnParameters SpawnParams;
+	SpawnParams.SpawnCollisionHandlingOverride = ESpawnActorCollisionHandlingMethod::AlwaysSpawn;
+	AVPTransientEditorTickableActorBase* NewActor = World->SpawnActor<AVPTransientEditorTickableActorBase>(ActorClass.Get(), Location, Rotation, SpawnParams);
+	return NewActor;
+}
diff --git a/Engine/Plugins/Experimental/VirtualProductionUtilities/Source/VPUtilitiesEditor/Public/VPEditorTickableActorBase.h b/Engine/Plugins/Experimental/VirtualProductionUtilities/Source/VPUtilitiesEditor/Public/VPEditorTickableActorBase.h
--- a/Engine/Plugins/Experimental/VirtualProductionUtilities/Source/VPUtilitiesEditor/Public/VPEditorTickableActorBase.h
+++ b/Engine/Plugins/Experimental/VirtualProductionUtilities/Source/VPUtilitiesEditor/Public/VPEditorTickableActorBase.h
@@ -11,7 +11,7 @@
 /**
  * Actor that tick in the Editor viewport with the event EditorTick.
  */
-UCLASS(Abstract, Transient)
+UCLASS(Abstract)
 class AVPEditorTickableActorBase : public AActor
 {
 	GENERATED_BODY()
diff --git a/Engine/Plugins/Experimental/VirtualProductionUtilities/Source/VPUtilitiesEditor/Public/VPTransientEditorTickableActorBase.h b/Engine/Plugins/Experimental/VirtualProductionUtilities/Source/VPUtilitiesEditor/Public/VPTransientEditorTickableActorBase.h
new file mode 100644
--- /dev/null
+++ b/Engine/Plugins/Experimental/VirtualProductionUtilities/Source/VPUtilitiesEditor/Public/VPTransientEditorTickableActorBase.h
@@ -0,0 +1,16 @@
+// Copyright 1998-2019 Epic Games, Inc. All Rights Reserved.
+
+#pragma once
+
+#include "VPEditorTickableActorBase.h"
+
+#include "VPTransientEditorTickableActorBase.generated.h"
+
+/**
+ * Specific VPEditorTickableActor explicitely marked as Transient
+ */
+UCLASS(Abstract, Transient)
+class AVPTransientEditorTickableActorBase : public AVPEditorTickableActorBase
+{
+	GENERATED_BODY()
+};
diff --git a/Engine/Plugins/Experimental/VirtualProductionUtilities/Source/VPUtilitiesEditor/Public/VPUtilitiesEditorBlueprintLibrary.h b/Engine/Plugins/Experimental/VirtualProductionUtilities/Source/VPUtilitiesEditor/Public/VPUtilitiesEditorBlueprintLibrary.h
--- a/Engine/Plugins/Experimental/VirtualProductionUtilities/Source/VPUtilitiesEditor/Public/VPUtilitiesEditorBlueprintLibrary.h
+++ b/Engine/Plugins/Experimental/VirtualProductionUtilities/Source/VPUtilitiesEditor/Public/VPUtilitiesEditorBlueprintLibrary.h
@@ -4,6 +4,7 @@
 
 #include "Kismet/BlueprintFunctionLibrary.h"
 #include "VPEditorTickableActorBase.h"
+#include "VPTransientEditorTickableActorBase.h"
 #include "VPUtilitiesEditorBlueprintLibrary.generated.h"
 
 
@@ -17,4 +18,7 @@ public:
 	UFUNCTION(BlueprintCallable, Category = "Virtual Production")
 	static AVPEditorTickableActorBase* SpawnVPEditorTickableActor(UObject* ContextObject, const TSubclassOf<AVPEditorTickableActorBase> ActorClass, const FVector Location, const FRotator Rotation);
 		
+	/** Spawn an editor-only Transient virtual production tickable actor */
+	UFUNCTION(BlueprintCallable, Category = "Virtual Production")
+	static AVPTransientEditorTickableActorBase* SpawnVPTransientEditorTickableActor(UObject* ContextObject, const TSubclassOf<AVPTransientEditorTickableActorBase> ActorClass, const FVector Location, const FRotator Rotation);
 };
\ No newline at end of file
diff --git a/Engine/Plugins/Experimental/VirtualProductionUtilities/Source/VPUtilitiesEditor/Public/VRTool.h b/Engine/Plugins/Experimental/VirtualProductionUtilities/Source/VPUtilitiesEditor/Public/VRTool.h
--- a/Engine/Plugins/Experimental/VirtualProductionUtilities/Source/VPUtilitiesEditor/Public/VRTool.h
+++ b/Engine/Plugins/Experimental/VirtualProductionUtilities/Source/VPUtilitiesEditor/Public/VRTool.h
@@ -3,14 +3,14 @@
 #pragma once
 
 #include "CoreMinimal.h"
-#include "VPEditorTickableActorBase.h"
+#include "VPTransientEditorTickableActorBase.h"
 #include "VRTool.generated.h"
 
 /**
  * Base class for virtual production VR tools
  */
 UCLASS()
-class AVRTool : public AVPEditorTickableActorBase
+class AVRTool : public AVPTransientEditorTickableActorBase
 {
 	GENERATED_BODY()
 
