- Fixing comments for VPUtilities spawn methods
- Fixing wrong spawn method used for VRTools

diff --git a/Engine/Plugins/Experimental/VirtualProductionUtilities/Source/VPUtilitiesEditor/Public/VPUtilitiesEditorBlueprintLibrary.h b/Engine/Plugins/Experimental/VirtualProductionUtilities/Source/VPUtilitiesEditor/Public/VPUtilitiesEditorBlueprintLibrary.h
--- a/Engine/Plugins/Experimental/VirtualProductionUtilities/Source/VPUtilitiesEditor/Public/VPUtilitiesEditorBlueprintLibrary.h
+++ b/Engine/Plugins/Experimental/VirtualProductionUtilities/Source/VPUtilitiesEditor/Public/VPUtilitiesEditorBlueprintLibrary.h
@@ -14,11 +14,19 @@ class VPUTILITIESEDITOR_API UVPUtilitiesEditorBlueprintLibrary : public UBluepri
 	GENERATED_BODY()
 
 public:
-	/** Spawn an editor-only virtual production tickable actor */
+	/**
+	 * Spawn an editor-only virtual production tickable actor 
+	 * @note Actors based on the non-transient AVPEditorTickableActorBase will be saved in the level. 
+	 * @note Being non-transient also means that transactions happening on them will be replicated on other connected multi-user machines
+	 */
 	UFUNCTION(BlueprintCallable, Category = "Virtual Production")
 	static AVPEditorTickableActorBase* SpawnVPEditorTickableActor(UObject* ContextObject, const TSubclassOf<AVPEditorTickableActorBase> ActorClass, const FVector Location, const FRotator Rotation);
 		
-	/** Spawn an editor-only Transient virtual production tickable actor */
+	/**
+	 * Spawn an editor-only transient virtual production tickable actor
+	 * @note Actors based on the transient AVPTransientEditorTickableActorBase will NOT be saved in the level.
+	 * @note Being transient also means that transactions happening on them will NOT be replicated on other connected multi-user machines
+	 */
 	UFUNCTION(BlueprintCallable, Category = "Virtual Production")
 	static AVPTransientEditorTickableActorBase* SpawnVPTransientEditorTickableActor(UObject* ContextObject, const TSubclassOf<AVPTransientEditorTickableActorBase> ActorClass, const FVector Location, const FRotator Rotation);
 };
\ No newline at end of file
