Integrating CL#9620606 from //UE4/Main to Dev-VirtualProduction (//UE4/Dev-VirtualProduction)
Fix for BP compile errors when opening VirtualCamera Sample

diff --git a/Engine/Plugins/Experimental/VirtualProductionUtilities/VirtualProductionUtilities.uplugin b/Engine/Plugins/Experimental/VirtualProductionUtilities/VirtualProductionUtilities.uplugin
--- a/Engine/Plugins/Experimental/VirtualProductionUtilities/VirtualProductionUtilities.uplugin
+++ b/Engine/Plugins/Experimental/VirtualProductionUtilities/VirtualProductionUtilities.uplugin
@@ -43,6 +43,10 @@
 		{
 			"Name": "EditorScriptingUtilities",
 			"Enabled": true
+		},
+		{
+			"Name": "SequencerScripting",
+			"Enabled": true
 		}
 	]
 }
