CL 9406511
Reduce filename size to fix can't write error

CL 9489134
Snapshot Multi User Material sync fix

CL 9489134
Sequencer tool now has corresponding icon on motion controller
Fixed bug where Vive UI was still appearing on Rift after deleting an object through the context menu
Help System UI no longer casts shadows

CL 9531557
Resave VP assets with UMG.

CL 9538549
Expose SnapshotTransactionBuffer to Blueprint as SnapshotObject

CL 9540413
Fixes VR on-camera viewfinders to have the correct color correction and gamma.
1) New shader, specifically applied to VR viewfinders, that performs a previously missing sRGB -> linear conversion.
2) Detect where we're in VR or not when creating camera preview widgets; don't apply tone mapping when we're in VR (to eliminate double tonemapping).

CL 9542364
Calling SnapshotObject when camera is attached to handheld/tripod so other users can see it update

CL 9551135
Add if (!IsRunningCommandlet()) to shader load to ensure it's ignored during cooking.

CL 9577516
Fixed Focusable windows causing no tool after menu close

CL 9601939
added TransactObject to Handheld and tripod tools

CL 9604635
Changed tripod and handheld tool to update the camera transform directly instead of using AttachToActor

CL 9605124
Live Editing cvar sequencer key mesh toggle

CL 9610074
Combined MarkerTool proc mesh into 1 draw call.

CL 9610330
Set camera label to "Playback Cam" when spawned by sequencer

diff --git a/Engine/Plugins/Experimental/VirtualProductionUtilities/Config/BaseVirtualProductionUtilities.ini b/Engine/Plugins/Experimental/VirtualProductionUtilities/Config/BaseVirtualProductionUtilities.ini
--- a/Engine/Plugins/Experimental/VirtualProductionUtilities/Config/BaseVirtualProductionUtilities.ini
+++ b/Engine/Plugins/Experimental/VirtualProductionUtilities/Config/BaseVirtualProductionUtilities.ini
@@ -4,4 +4,7 @@
 
 
 [/Script/VPUtilitiesEditor.VPUtilitiesEditorSettings]
-ScoutingSubsystemEdititorUtilityActorClassPath=/VirtualProductionUtilities/VirtualProductionHelpers.VirtualProductionHelpers_C
+ScoutingSubsystemEditorUtilityClassPath=/VirtualProductionUtilities/BP_VPScoutingSubsystemHelpers.BP_VPScoutingSubsystemHelpers_C
+GestureManagerEditorUtilityClassPath=/VirtualProductionUtilities/VRTools/BP_GestureManager.BP_GestureManager_C
++AdditionnalClassToLoad=/VirtualProductionUtilities/VR/VirtualScoutingInteractor.VirtualScoutingInteractor
++AdditionnalClassToLoad=/VirtualProductionUtilities/VR/VirtualScoutingTeleporter.VirtualScoutingTeleporter
\ No newline at end of file
diff --git a/Engine/Plugins/Experimental/VirtualProductionUtilities/Source/VPUtilitiesEditor/Private/VPScoutingSubsystem.cpp b/Engine/Plugins/Experimental/VirtualProductionUtilities/Source/VPUtilitiesEditor/Private/VPScoutingSubsystem.cpp
--- a/Engine/Plugins/Experimental/VirtualProductionUtilities/Source/VPUtilitiesEditor/Private/VPScoutingSubsystem.cpp
+++ b/Engine/Plugins/Experimental/VirtualProductionUtilities/Source/VPUtilitiesEditor/Private/VPScoutingSubsystem.cpp
@@ -10,12 +10,15 @@
 #include "WidgetBlueprint.h"
 #include "EditorUtilityActor.h"
 #include "EditorUtilityWidget.h"
+#include "Engine/AssetManager.h"
 #include "IVREditorModule.h"
 #include "UObject/ConstructorHelpers.h"
+#include "UObject/Script.h"
 #include "VPSettings.h"
 #include "VPUtilitiesEditorSettings.h"
-#include "LevelEditor.h"
 
+/* UVPScoutingSubsystem name
+ *****************************************************************************/
 const FName UVPScoutingSubsystem::VProdPanelID = FName(TEXT("VirtualProductionPanel"));
 const FName UVPScoutingSubsystem::VProdPanelLeftID = FName(TEXT("VirtualProductionPanelLeft"));
 const FName UVPScoutingSubsystem::VProdPanelRightID = FName(TEXT("VirtualProductionPanelRight"));
@@ -24,6 +27,78 @@ const FName UVPScoutingSubsystem::VProdPanelTimelineID = FName(TEXT("VirtualProd
 const FName UVPScoutingSubsystem::VProdPanelMeasureID = FName(TEXT("VirtualProductionPanelMeasure"));
 const FName UVPScoutingSubsystem::VProdPanelGafferID = FName(TEXT("VirtualProductionPanelGaffer"));
 
+
+/* UVPScoutingSubsystemGestureManagerBase
+ *****************************************************************************/
+UVPScoutingSubsystemGestureManagerBase::UVPScoutingSubsystemGestureManagerBase()
+{
+	if (!HasAnyFlags(RF_ArchetypeObject | RF_ClassDefaultObject))
+	{
+		IVREditorModule::Get().OnVREditingModeEnter().AddUObject(this, &UVPScoutingSubsystemGestureManagerBase::OnVREditingModeEnterCallback);
+		IVREditorModule::Get().OnVREditingModeExit().AddUObject(this, &UVPScoutingSubsystemGestureManagerBase::OnVREditingModeExitCallback);
+	}
+}
+
+void UVPScoutingSubsystemGestureManagerBase::BeginDestroy()
+{
+	if (!HasAnyFlags(RF_ArchetypeObject | RF_ClassDefaultObject))
+	{
+		IVREditorModule::Get().OnVREditingModeEnter().RemoveAll(this);
+		IVREditorModule::Get().OnVREditingModeExit().RemoveAll(this);
+	}
+
+	Super::BeginDestroy();
+}
+
+void UVPScoutingSubsystemGestureManagerBase::Tick(float DeltaTime)
+{
+	FEditorScriptExecutionGuard ScriptGuard;
+	EditorTick(DeltaTime);
+}
+
+bool UVPScoutingSubsystemGestureManagerBase::IsTickable() const
+{
+	if (IVREditorModule::IsAvailable())
+	{
+		return IVREditorModule::Get().IsVREditorModeActive();
+	}
+	return false;
+}
+
+TStatId UVPScoutingSubsystemGestureManagerBase::GetStatId() const
+{
+	RETURN_QUICK_DECLARE_CYCLE_STAT(UVPScoutingSubsystemGestureManagerBase, STATGROUP_Tickables);
+}
+
+void UVPScoutingSubsystemGestureManagerBase::OnVREditingModeEnterCallback()
+{
+	FEditorScriptExecutionGuard ScriptGuard;
+	OnVREditingModeEnter();
+}
+
+void UVPScoutingSubsystemGestureManagerBase::OnVREditingModeExitCallback()
+{
+	FEditorScriptExecutionGuard ScriptGuard;
+	OnVREditingModeExit();
+}
+
+void UVPScoutingSubsystemGestureManagerBase::EditorTick_Implementation(float DeltaSeconds)
+{
+
+}
+
+void UVPScoutingSubsystemGestureManagerBase::OnVREditingModeEnter_Implementation()
+{
+
+}
+
+void UVPScoutingSubsystemGestureManagerBase::OnVREditingModeExit_Implementation()
+{
+
+}
+
+/* UVPScoutingSubsystem
+ *****************************************************************************/
 UVPScoutingSubsystem::UVPScoutingSubsystem()
 	: UEditorSubsystem()
 {
@@ -56,45 +131,48 @@ void UVPScoutingSubsystem::Initialize(FSubsystemCollectionBase& Collection)
 		CVar->Set(0);
 	}
 
+	// Load the ScoutingHelper implemented in BP. See BaseVirtualProductionUtilitites.ini
+	VPSubsystemHelpers = nullptr;
+	if (UClass* EditorUtilityClass = GetDefault<UVPUtilitiesEditorSettings>()->ScoutingSubsystemEditorUtilityClassPath.TryLoadClass<UVPScoutingSubsystemHelpersBase>())
+	{
+		VPSubsystemHelpers = NewObject<UVPScoutingSubsystemHelpersBase>(GetTransientPackage(), EditorUtilityClass);
+	}
+	else
+	{
+		UE_LOG(LogVPUtilitiesEditor, Warning, TEXT("Failed loading VPScoutingHelpers \"%s\""), *GetDefault<UVPUtilitiesEditorSettings>()->ScoutingSubsystemEditorUtilityClassPath.ToString());
+	}
+
 	// to do final initializations at the right time
 	EngineInitCompleteDelegate = FCoreDelegates::OnFEngineLoopInitComplete.AddUObject(this, &UVPScoutingSubsystem::OnEngineInitComplete);
 }
 
 void UVPScoutingSubsystem::Deinitialize()
 {
-	FLevelEditorModule& LevelEditorModule = FModuleManager::LoadModuleChecked<FLevelEditorModule>(TEXT("LevelEditor"));
-	LevelEditorModule.OnMapChanged().RemoveAll(this);
+	VPSubsystemHelpers = nullptr;
 }
 
 void UVPScoutingSubsystem::OnEngineInitComplete()
 {
-	FLevelEditorModule& LevelEditorModule = FModuleManager::LoadModuleChecked<FLevelEditorModule>(TEXT("LevelEditor"));
-	LevelEditorModule.OnMapChanged().AddUObject(this, &UVPScoutingSubsystem::OnMapChanged);
-
 	FCoreDelegates::OnFEngineLoopInitComplete.Remove(EngineInitCompleteDelegate);
 	EngineInitCompleteDelegate.Reset();
-}
 
-void UVPScoutingSubsystem::OnMapChanged(UWorld * World, EMapChangeType MapChangeType)
-{
-	if (MapChangeType == EMapChangeType::TearDownWorld)
+	// Load the GestureManager implemented in BP. See BaseVirtualProductionUtilitites.ini
+	// GestureManager needs the Take module, load it once the engine is loaded.
+	GestureManager = nullptr;
+	if (UClass* EditorUtilityClass = GetDefault<UVPUtilitiesEditorSettings>()->GestureManagerEditorUtilityClassPath.TryLoadClass<UVPScoutingSubsystemGestureManagerBase>())
 	{
-		VProdHelper = nullptr;
+		GestureManager = NewObject<UVPScoutingSubsystemGestureManagerBase>(GetTransientPackage(), EditorUtilityClass);
 	}
-	else if (MapChangeType == EMapChangeType::LoadMap || MapChangeType == EMapChangeType::NewMap)
+	else
 	{
-		const UVPUtilitiesEditorSettings* VPUtilitiesEditorSettings = GetDefault<UVPUtilitiesEditorSettings>();
-		const FString ClassPath = VPUtilitiesEditorSettings->ScoutingSubsystemEdititorUtilityActorClassPath.ToString();
-		UClass* EditorUtilityActorClass = LoadObject<UClass>(nullptr, *ClassPath);
+		UE_LOG(LogVPUtilitiesEditor, Warning, TEXT("Failed loading VPScoutingHelpers \"%s\""), *GetDefault<UVPUtilitiesEditorSettings>()->GestureManagerEditorUtilityClassPath.ToString());
+	}
 
-		if (EditorUtilityActorClass)
-		{
-			VProdHelper = NewObject<AEditorUtilityActor>(GetTransientPackage(), EditorUtilityActorClass);
-		}
-		else
-		{
-			UE_LOG(LogVPUtilitiesEditor, Warning, TEXT("Failed loading EditorUtilityActorClass \"%s\""), *ClassPath);
-		}
+	// In debug some asset take a long time to load and crash the engine, preload those asset in async mode to prevent that
+	for (const FSoftClassPath& ClassAssetPath : GetDefault<UVPUtilitiesEditorSettings>()->AdditionnalClassToLoad)
+	{
+		FStreamableManager& StreamableManager = UAssetManager::Get().GetStreamableManager();
+		StreamableManager.RequestAsyncLoad(ClassAssetPath);
 	}
 }
 
@@ -303,3 +381,12 @@ void UVPScoutingSubsystem::SetIsHelperSystemEnabled(const bool bInIsHelperSystem
 	VPUtilitiesEditorSettings->bIsHelperSystemEnabled = bInIsHelperSystemEnabled;
 	VPUtilitiesEditorSettings->SaveConfig();
 }
+
+void UVPScoutingSubsystem::ExitVRMode()
+{
+	IVREditorModule& VREditorModule = IVREditorModule::Get();
+	if (VREditorModule.IsVREditorEnabled())
+	{
+		VREditorModule.EnableVREditor(false);
+	}
+}
diff --git a/Engine/Plugins/Experimental/VirtualProductionUtilities/Source/VPUtilitiesEditor/Private/VPUtilitiesEditorBlueprintLibrary.cpp b/Engine/Plugins/Experimental/VirtualProductionUtilities/Source/VPUtilitiesEditor/Private/VPUtilitiesEditorBlueprintLibrary.cpp
--- a/Engine/Plugins/Experimental/VirtualProductionUtilities/Source/VPUtilitiesEditor/Private/VPUtilitiesEditorBlueprintLibrary.cpp
+++ b/Engine/Plugins/Experimental/VirtualProductionUtilities/Source/VPUtilitiesEditor/Private/VPUtilitiesEditorBlueprintLibrary.cpp
@@ -56,7 +56,7 @@ AVPTransientEditorTickableActorBase* UVPUtilitiesEditorBlueprintLibrary::SpawnVP
 	return NewActor;
 }
 
-UTexture* UVPUtilitiesEditorBlueprintLibrary::ImportSnapshotTexture(FString FileName, FString SubFolderName)
+UTexture* UVPUtilitiesEditorBlueprintLibrary::ImportSnapshotTexture(FString FileName, FString SubFolderName, FString AbsolutePathPackage)
 {
 	UTexture* UnrealTexture = NULL;
 
@@ -65,13 +65,16 @@ UTexture* UVPUtilitiesEditorBlueprintLibrary::ImportSnapshotTexture(FString File
 		return UnrealTexture;
 	}
 
-	FString AbsolutePathPackage = FPaths::ProjectSavedDir() + "VirtualProduction/Snapshots/" + SubFolderName + "/";
+	if (AbsolutePathPackage.IsEmpty())
+	{
+		AbsolutePathPackage = FPaths::ProjectSavedDir() + "VirtualProduction/Snapshots/" + SubFolderName + "/";
+	}
 
 	FString TextureName = FPaths::GetBaseFilename(FileName);
 	TextureName = ObjectTools::SanitizeObjectName(TextureName);
 	FString Extension = FPaths::GetExtension(FileName).ToLower();
 
-	FString PackageName = TEXT("/Game/VirtualProduction/Snapshots/" + SubFolderName + "/");
+	FString PackageName = TEXT("/Game/Snapshots/" + SubFolderName + "/");
 	PackageName += TextureName;
 	UPackage* Package = CreatePackage(NULL, *PackageName);
 	Package->FullyLoad();
diff --git a/Engine/Plugins/Experimental/VirtualProductionUtilities/Source/VPUtilitiesEditor/Public/VPScoutingSubsystem.h b/Engine/Plugins/Experimental/VirtualProductionUtilities/Source/VPUtilitiesEditor/Public/VPScoutingSubsystem.h
--- a/Engine/Plugins/Experimental/VirtualProductionUtilities/Source/VPUtilitiesEditor/Public/VPScoutingSubsystem.h
+++ b/Engine/Plugins/Experimental/VirtualProductionUtilities/Source/VPUtilitiesEditor/Public/VPScoutingSubsystem.h
@@ -5,6 +5,7 @@
 #include "EditorSubsystem.h"
 #include "EditorUtilityWidget.h"
 #include "Templates/SubclassOf.h"
+#include "TickableEditorObject.h"
 #include "UI/VREditorFloatingUI.h"
 #include "UnrealEdMisc.h"
 #include "VPScoutingSubsystem.generated.h"
@@ -22,13 +23,62 @@ enum class EVProdPanelIDs : uint8
 	Gaffer
 };
 
+
+/*
+ * Base class of the helper class defined in BP
+ */
+UCLASS(Abstract, Blueprintable, MinimalAPI, meta = (ShowWorldContextPin))
+class UVPScoutingSubsystemHelpersBase : public UObject
+{
+	GENERATED_BODY()
+};
+
+
+/*
+ * Base class of the gesture manager defined in BP
+ */
+UCLASS(Abstract, Blueprintable, MinimalAPI, meta = (ShowWorldContextPin))
+class UVPScoutingSubsystemGestureManagerBase : public UObject, public FTickableEditorObject
+{
+	GENERATED_BODY()
+
+public:
+	UVPScoutingSubsystemGestureManagerBase();
+
+	UFUNCTION(BlueprintNativeEvent, CallInEditor, BlueprintCallable, Category = "Tick")
+	void EditorTick(float DeltaSeconds);
+
+	UFUNCTION(BlueprintNativeEvent, CallInEditor, BlueprintCallable, Category = "VR")
+	void OnVREditingModeEnter();
+	UFUNCTION(BlueprintNativeEvent, CallInEditor, BlueprintCallable, Category = "VR")
+	void OnVREditingModeExit();
+
+	//~ Begin UObject interface
+	virtual void BeginDestroy() override;
+	//~ End UObject interface
+
+	//~ Begin FTickableEditorObject interface
+	virtual void Tick(float DeltaTime) override;
+	virtual ETickableTickType GetTickableTickType() const override { return ETickableTickType::Conditional; }
+	virtual bool IsTickable() const override;
+	virtual TStatId GetStatId() const override;
+	//~ End FTickableEditorObject interface
+
+private:
+	void OnVREditingModeEnterCallback();
+	void OnVREditingModeExitCallback();
+};
+
+
+/*
+ * Subsystem used for VR Scouting
+ */
 UCLASS()
 class UVPScoutingSubsystem : public UEditorSubsystem
 {
 	GENERATED_BODY()
 
 public:
-
 	UVPScoutingSubsystem();
 	
 	virtual void Initialize(FSubsystemCollectionBase& Collection) override;
@@ -36,7 +86,11 @@ public:
 
 	/** Subsystems can't have any Blueprint implementations, so we attach this class for any BP logic that we to provide. */
 	UPROPERTY(Transient, BlueprintReadOnly, Category = "Virtual Production")
-	class AEditorUtilityActor* VProdHelper;
+	UVPScoutingSubsystemHelpersBase* VPSubsystemHelpers;
+
+	/** GestureManager that manage some user input in VR editor. */
+	UPROPERTY(Transient, BlueprintReadOnly, Category = "Virtual Production")
+	UVPScoutingSubsystemGestureManagerBase* GestureManager;
 
 	/** bool to keep track of whether the settings menu panel in the main menu is open*/
 	UPROPERTY(EditAnywhere, BlueprintReadWrite, Category = "Main Menu")
@@ -148,14 +202,15 @@ public:
 	UFUNCTION(BlueprintCallable, Category = "Virtual Production")
 	static void SetIsHelperSystemEnabled(const bool bInIsHelperSystemEnabled);
 
+	/** Exit VR Mode  */
+	UFUNCTION(BlueprintCallable, Category = "Virtual Production")
+	static void ExitVRMode();
+
 private:
 
 	FDelegateHandle EngineInitCompleteDelegate;
 
 	void OnEngineInitComplete();
-
-	/** Callback from the level editor when the map changes */
-	void OnMapChanged(UWorld* World, EMapChangeType MapChangeType);
 	
 	// Static IDs when submitting open/close requests for the VProd main menu panels. VREditorUISystem uses FNames to manage its panels, so these should be used for consistency.	
 	static const FName VProdPanelID;
diff --git a/Engine/Plugins/Experimental/VirtualProductionUtilities/Source/VPUtilitiesEditor/Public/VPUtilitiesEditorBlueprintLibrary.h b/Engine/Plugins/Experimental/VirtualProductionUtilities/Source/VPUtilitiesEditor/Public/VPUtilitiesEditorBlueprintLibrary.h
--- a/Engine/Plugins/Experimental/VirtualProductionUtilities/Source/VPUtilitiesEditor/Public/VPUtilitiesEditorBlueprintLibrary.h
+++ b/Engine/Plugins/Experimental/VirtualProductionUtilities/Source/VPUtilitiesEditor/Public/VPUtilitiesEditorBlueprintLibrary.h
@@ -33,5 +33,5 @@ public:
 
 	/** Imports Image file into VirtualProduction/Snapshots/ folder */
 	UFUNCTION(BlueprintCallable, Category = "Virtual Production")
-	static UTexture* ImportSnapshotTexture(FString FileName, FString SubFolderName);
+	static UTexture* ImportSnapshotTexture(FString FileName, FString SubFolderName, FString AbsolutePathPackage);
 };
\ No newline at end of file
diff --git a/Engine/Plugins/Experimental/VirtualProductionUtilities/Source/VPUtilitiesEditor/Public/VPUtilitiesEditorSettings.h b/Engine/Plugins/Experimental/VirtualProductionUtilities/Source/VPUtilitiesEditor/Public/VPUtilitiesEditorSettings.h
--- a/Engine/Plugins/Experimental/VirtualProductionUtilities/Source/VPUtilitiesEditor/Public/VPUtilitiesEditorSettings.h
+++ b/Engine/Plugins/Experimental/VirtualProductionUtilities/Source/VPUtilitiesEditor/Public/VPUtilitiesEditorSettings.h
@@ -52,8 +52,16 @@ public:
 	bool bIsHelperSystemEnabled = true;
 	
 	/** ScoutingSubsystem class to use for Blueprint helpers */
+	UPROPERTY(config, meta = (MetaClass = "VPScoutingSubsystemHelpersBase"))
+	FSoftClassPath ScoutingSubsystemEditorUtilityClassPath;
+
+	/** GestureManager class to use by the ScoutingSubsystem */
+	UPROPERTY(config, meta = (MetaClass = "VPScoutingSubsystemGestureManagerBase"))
+	FSoftClassPath GestureManagerEditorUtilityClassPath;
+
+	/** GestureManager class to use by the ScoutingSubsystem */
 	UPROPERTY(config)
-	FSoftClassPath ScoutingSubsystemEdititorUtilityActorClassPath;
+	TArray<FSoftClassPath> AdditionnalClassToLoad;
 
 public:
 
diff --git a/Engine/Source/Editor/LevelEditor/Private/SLevelViewport.cpp b/Engine/Source/Editor/LevelEditor/Private/SLevelViewport.cpp
--- a/Engine/Source/Editor/LevelEditor/Private/SLevelViewport.cpp
+++ b/Engine/Source/Editor/LevelEditor/Private/SLevelViewport.cpp
@@ -3483,6 +3483,10 @@ void SLevelViewport::PreviewActors( const TArray< AActor* >& InActorsToPreview,
 				ActorPreviewLevelViewportClient->EngineShowFlags = FEngineShowFlags(ESFIM_Game);
 				ActorPreviewLevelViewportClient->EngineShowFlags.SetSelection(true);
 				ActorPreviewLevelViewportClient->LastEngineShowFlags = FEngineShowFlags(ESFIM_Editor);
+				if (!bPreviewInDesktopViewport)
+				{
+					ActorPreviewLevelViewportClient->EngineShowFlags.Tonemapper = false;
+				}
 
 				// We don't use view modes for preview viewports
 				ActorPreviewLevelViewportClient->SetViewMode(VMI_Unknown);
diff --git a/Engine/Source/Editor/Sequencer/Private/SequencerEdMode.cpp b/Engine/Source/Editor/Sequencer/Private/SequencerEdMode.cpp
--- a/Engine/Source/Editor/Sequencer/Private/SequencerEdMode.cpp
+++ b/Engine/Source/Editor/Sequencer/Private/SequencerEdMode.cpp
@@ -26,6 +26,11 @@
 
 const FEditorModeID FSequencerEdMode::EM_SequencerMode(TEXT("EM_SequencerMode"));
 
+static TAutoConsoleVariable<bool> CVarDrawMeshTrails(
+	TEXT("Sequencer.DrawMeshTrails"),
+	true,
+	TEXT("Toggle to show or hide Level Sequencer VR Editor trails"));
+
 FSequencerEdMode::FSequencerEdMode()
 {
 	FSequencerEdModeTool* SequencerEdModeTool = new FSequencerEdModeTool(this);
@@ -33,12 +38,16 @@ FSequencerEdMode::FSequencerEdMode()
 	Tools.Add( SequencerEdModeTool );
 	SetCurrentTool( SequencerEdModeTool );
 
-	// todo vreditor: make this a setting
-	bDrawMeshTrails = true;
+	bDrawMeshTrails = CVarDrawMeshTrails->GetBool();
+	CVarDrawMeshTrails.AsVariable()->SetOnChangedCallback(FConsoleVariableDelegate::CreateLambda([this](IConsoleVariable* Var) 
+	{
+		bDrawMeshTrails = Var->GetBool();
+	}));
 }
 
 FSequencerEdMode::~FSequencerEdMode()
 {
+	CVarDrawMeshTrails.AsVariable()->SetOnChangedCallback(FConsoleVariableDelegate());
 }
 
 void FSequencerEdMode::Enter()
diff --git a/Engine/Source/Editor/VREditor/UI/VREditorDockableWindow.cpp b/Engine/Source/Editor/VREditor/UI/VREditorDockableWindow.cpp
--- a/Engine/Source/Editor/VREditor/UI/VREditorDockableWindow.cpp
+++ b/Engine/Source/Editor/VREditor/UI/VREditorDockableWindow.cpp
@@ -24,8 +24,8 @@ namespace VREd
 }
 
 
-AVREditorDockableWindow::AVREditorDockableWindow() : 
-	Super(),
+AVREditorDockableWindow::AVREditorDockableWindow(const FObjectInitializer& ObjectInitializer) :
+	Super(ObjectInitializer),
 	bIsLaserAimingTowardUI(false),
 	AimingAtMeFadeAlpha(0.0f),
 	bIsHoveringOverSelectionBar(false),
diff --git a/Engine/Source/Editor/VREditor/UI/VREditorDockableWindow.h b/Engine/Source/Editor/VREditor/UI/VREditorDockableWindow.h
--- a/Engine/Source/Editor/VREditor/UI/VREditorDockableWindow.h
+++ b/Engine/Source/Editor/VREditor/UI/VREditorDockableWindow.h
@@ -23,7 +23,7 @@ class AVREditorDockableWindow : public AVREditorFloatingUI, public IViewportInte
 public:
 
 	/** Default constructor */
-	AVREditorDockableWindow();
+	AVREditorDockableWindow(const FObjectInitializer& ObjectInitializer);
 	virtual void PostActorCreated() override;
 
 	/** Updates the meshes for the UI */
diff --git a/Engine/Source/Editor/VREditor/UI/VREditorFloatingCameraUI.cpp b/Engine/Source/Editor/VREditor/UI/VREditorFloatingCameraUI.cpp
--- a/Engine/Source/Editor/VREditor/UI/VREditorFloatingCameraUI.cpp
+++ b/Engine/Source/Editor/VREditor/UI/VREditorFloatingCameraUI.cpp
@@ -4,8 +4,7 @@
 #include "VREditorUISystem.h"
 #include "VREditorBaseUserWidget.h"
 #include "VREditorMode.h"
-#include "Components/WidgetComponent.h"
-#include "VREditorWidgetComponent.h"
+#include "VREditorCameraWidgetComponent.h"
 #include "Components/StaticMeshComponent.h"
 #include "VRModeSettings.h"
 #include "VREditorAssetContainer.h"
@@ -16,8 +15,9 @@
 
 #define LOCTEXT_NAMESPACE "AVREditorFloatingCameraUI"
 
-AVREditorFloatingCameraUI::AVREditorFloatingCameraUI()
-	: Super()
+AVREditorFloatingCameraUI::AVREditorFloatingCameraUI(const FObjectInitializer& ObjectInitializer)
+	: Super(ObjectInitializer.SetDefaultSubobjectClass<UVREditorCameraWidgetComponent>(TEXT("WidgetComponent"))),
+	OffsetFromCamera( -25.0f, 0.0f, 30.0f )
 {
 	if (HasAnyFlags(RF_ClassDefaultObject))
 	{
@@ -46,8 +46,7 @@ FTransform AVREditorFloatingCameraUI::MakeCustomUITransform()
 		CameraTransform = LinkedActor->GetTransform();
 
 		const FTransform UIFlipTransform(FRotator(0.0f, 180.0f, 0.0f).Quaternion(), FVector::ZeroVector);
-		const FVector Offset = FVector(-25.0f, 0.0f, 40.0f);
-		const FTransform OffsetTransform(FRotator::ZeroRotator, Offset);
+		const FTransform OffsetTransform(FRotator::ZeroRotator, OffsetFromCamera);
 
 		UITransform = UIFlipTransform * OffsetTransform * CameraTransform;
 	}
diff --git a/Engine/Source/Editor/VREditor/UI/VREditorFloatingCameraUI.h b/Engine/Source/Editor/VREditor/UI/VREditorFloatingCameraUI.h
--- a/Engine/Source/Editor/VREditor/UI/VREditorFloatingCameraUI.h
+++ b/Engine/Source/Editor/VREditor/UI/VREditorFloatingCameraUI.h
@@ -22,7 +22,12 @@ class AVREditorFloatingCameraUI : public AVREditorFloatingUI
 	GENERATED_BODY()
 
 public:
-	AVREditorFloatingCameraUI();
+
+	/** The offset of this UI from its camera */
+	UPROPERTY(EditDefaultsOnly, Category = "FloatingCameraUI")
+	FVector OffsetFromCamera;
+
+	AVREditorFloatingCameraUI(const FObjectInitializer& ObjectInitializer);
 	void SetLinkedActor(class AActor* InActor);
 
 	virtual FTransform MakeCustomUITransform() override;
diff --git a/Engine/Source/Editor/VREditor/UI/VREditorFloatingUI.cpp b/Engine/Source/Editor/VREditor/UI/VREditorFloatingUI.cpp
--- a/Engine/Source/Editor/VREditor/UI/VREditorFloatingUI.cpp
+++ b/Engine/Source/Editor/VREditor/UI/VREditorFloatingUI.cpp
@@ -14,8 +14,7 @@ namespace VREd
 	static FAutoConsoleVariable UIFadeSpeed(TEXT("VREd.UIFadeSpeed"), 6.0f, TEXT("How fast UI should fade in and out"));
 }
 
-
-AVREditorFloatingUI::AVREditorFloatingUI()
+AVREditorFloatingUI::AVREditorFloatingUI(const FObjectInitializer& ObjectInitializer)
 	: Super(),
 	SlateWidget(nullptr),
 	UserWidget(nullptr),
diff --git a/Engine/Source/Editor/VREditor/UI/VREditorFloatingUI.h b/Engine/Source/Editor/VREditor/UI/VREditorFloatingUI.h
--- a/Engine/Source/Editor/VREditor/UI/VREditorFloatingUI.h
+++ b/Engine/Source/Editor/VREditor/UI/VREditorFloatingUI.h
@@ -77,7 +77,7 @@ class AVREditorFloatingUI : public AVREditorBaseActor
 public:
 
 	/** Default constructor which sets up safe defaults */
-	AVREditorFloatingUI();
+	AVREditorFloatingUI(const FObjectInitializer& ObjectInitializer);
 
 	/** Creates a FVREditorFloatingUI using a Slate widget, and sets up safe defaults */
 	void SetSlateWidget( class UVREditorUISystem& InitOwner, const VREditorPanelID& InID, const TSharedRef<SWidget>& InitSlateWidget, const FIntPoint InitResolution, const float InitScale, const EDockedTo InitDockedTo );
@@ -180,7 +180,7 @@ public:
 
 	/** All params used to create this panel if this panel has a UMG widget and was created via BP. Invalid otherwise. */
 	UPROPERTY()
-	FVREditorFloatingUICreationContext CreationContext;
+	FVREditorFloatingUICreationContext CreationContext;	
 
 protected:
 
diff --git a/Engine/Source/Editor/VREditor/VREditorCameraWidgetComponent.cpp b/Engine/Source/Editor/VREditor/VREditorCameraWidgetComponent.cpp
new file mode 100644
--- /dev/null
+++ b/Engine/Source/Editor/VREditor/VREditorCameraWidgetComponent.cpp
@@ -0,0 +1,17 @@
+// Copyright 1998-2019 Epic Games, Inc. All Rights Reserved.
+
+#include "VREditorCameraWidgetComponent.h"
+#include "UObject/ConstructorHelpers.h"
+
+UVREditorCameraWidgetComponent::UVREditorCameraWidgetComponent()
+	: Super()	
+{	
+	// Override this shader for VR camera viewfinders so that we get color-correct images.
+	// This shader does an sRGB -> Linear conversion and doesn't apply the "UI Brightness" setting .
+	if (!IsRunningCommandlet())
+	{
+		static ConstructorHelpers::FObjectFinder<UMaterialInterface> OpaqueMaterial_OneSided_Finder(TEXT("/Engine/EngineMaterials/Widget3DCameraPassThrough_Opaque_OneSided"));
+		OpaqueMaterial_OneSided = OpaqueMaterial_OneSided_Finder.Object;
+	}
+}
+
diff --git a/Engine/Source/Editor/VREditor/VREditorCameraWidgetComponent.h b/Engine/Source/Editor/VREditor/VREditorCameraWidgetComponent.h
new file mode 100644
--- /dev/null
+++ b/Engine/Source/Editor/VREditor/VREditorCameraWidgetComponent.h
@@ -0,0 +1,19 @@
+// Copyright 1998-2019 Epic Games, Inc. All Rights Reserved.
+
+#pragma once
+
+#include "VREditorWidgetComponent.h"
+#include "VREditorCameraWidgetComponent.generated.h"
+
+/**
+ * A specialized WidgetComponent that color-corrects cameras previews (viewfinders) in VR
+ */
+UCLASS(hidedropdown)
+class UVREditorCameraWidgetComponent : public UVREditorWidgetComponent
+{
+	GENERATED_BODY()
+
+public:	
+	UVREditorCameraWidgetComponent();
+	
+};
