Changed GetFieldOfView node for GetHorizontalFieldOfView because the first one doesn't update when changing lenses.

Moved setting of cvars from VPScoutingSubsystem.cpp to virtual scouting interactor blueprint so it doesn't affect default VR Mode.

diff --git a/Engine/Plugins/Experimental/VirtualProductionUtilities/Source/VPUtilitiesEditor/Private/VPScoutingSubsystem.cpp b/Engine/Plugins/Experimental/VirtualProductionUtilities/Source/VPUtilitiesEditor/Private/VPScoutingSubsystem.cpp
--- a/Engine/Plugins/Experimental/VirtualProductionUtilities/Source/VPUtilitiesEditor/Private/VPScoutingSubsystem.cpp
+++ b/Engine/Plugins/Experimental/VirtualProductionUtilities/Source/VPUtilitiesEditor/Private/VPScoutingSubsystem.cpp
@@ -100,7 +100,8 @@ void UVPScoutingSubsystemGestureManagerBase::OnVREditingModeExit_Implementation(
 /* UVPScoutingSubsystem
  *****************************************************************************/
 UVPScoutingSubsystem::UVPScoutingSubsystem()
-	: UEditorSubsystem()
+	: UEditorSubsystem(),
+	GripNavSpeedCoeff(4.0f)
 {
 }
 
@@ -108,29 +109,6 @@ void UVPScoutingSubsystem::Initialize(FSubsystemCollectionBase& Collection)
 {
 	UE_LOG(LogVPUtilitiesEditor, Log, TEXT("VP Scouting subsystem initialized."));
 
-	// Setup VR editor settings from user preferences
-	UVPUtilitiesEditorSettings* VPUtilitiesEditorSettings = GetMutableDefault<UVPUtilitiesEditorSettings>();
-	
-	// Turn on/off transform VR gizmo
-	IConsoleVariable* CVar = IConsoleManager::Get().FindConsoleVariable(TEXT("VI.ShowTransformGizmo"));
-	CVar->Set(VPUtilitiesEditorSettings->bUseTransformGizmo);
-
-	//Initialize drag scale from saved config file
-	GripNavSpeedCoeff = 4.0f;
-	CVar = IConsoleManager::Get().FindConsoleVariable(TEXT("VI.DragScale"));
-	CVar->Set(VPUtilitiesEditorSettings->GripNavSpeed * GripNavSpeedCoeff);
-
-	//Turn on/off grip nav inertia
-	CVar = IConsoleManager::Get().FindConsoleVariable(TEXT("VI.HighSpeedInertiaDamping"));
-	if (VPUtilitiesEditorSettings->bUseGripInertiaDamping)
-	{
-		CVar->Set(VPUtilitiesEditorSettings->InertiaDamping);
-	}
-	else
-	{
-		CVar->Set(0);
-	}
-
 	// Load the ScoutingHelper implemented in BP. See BaseVirtualProductionUtilitites.ini
 	VPSubsystemHelpers = nullptr;
 	if (UClass* EditorUtilityClass = GetDefault<UVPUtilitiesEditorSettings>()->ScoutingSubsystemEditorUtilityClassPath.TryLoadClass<UVPScoutingSubsystemHelpersBase>())
@@ -318,12 +296,17 @@ void UVPScoutingSubsystem::SetIsUsingTransformGizmo(const bool bInIsUsingTransfo
 	if (bInIsUsingTransformGizmo != VPUtilitiesEditorSettings->bUseTransformGizmo)
 	{
 		VPUtilitiesEditorSettings->bUseTransformGizmo = bInIsUsingTransformGizmo;
-		IConsoleVariable* CVar = IConsoleManager::Get().FindConsoleVariable(TEXT("VI.ShowTransformGizmo"));
-		CVar->Set(bInIsUsingTransformGizmo);
+		SetShowTransformGizmoCVar(bInIsUsingTransformGizmo);
 		VPUtilitiesEditorSettings->SaveConfig();
 	}
 }
 
+void UVPScoutingSubsystem::SetShowTransformGizmoCVar(const bool bInShowTransformGizmoCVar)
+{
+	IConsoleVariable* CVar = IConsoleManager::Get().FindConsoleVariable(TEXT("VI.ShowTransformGizmo"));
+	CVar->Set(bInShowTransformGizmoCVar);
+}
+
 float UVPScoutingSubsystem::GetFlightSpeed()
 {
 	return GetDefault<UVPUtilitiesEditorSettings>()->FlightSpeed;
@@ -358,18 +341,24 @@ void UVPScoutingSubsystem::SetIsUsingInertiaDamping(const bool bInIsUsingInertia
 	//Save this value in editor settings and set the console variable which is used for inertia damping
 	UVPUtilitiesEditorSettings* VPUtilitiesEditorSettings = GetMutableDefault<UVPUtilitiesEditorSettings>();
 	VPUtilitiesEditorSettings->bUseGripInertiaDamping = bInIsUsingInertiaDamping;
-	IConsoleVariable* CVar = IConsoleManager::Get().FindConsoleVariable(TEXT("VI.HighSpeedInertiaDamping"));
+	
 	if (bInIsUsingInertiaDamping)
 	{
-		CVar->Set(VPUtilitiesEditorSettings->InertiaDamping);
+		SetInertiaDampingCVar(VPUtilitiesEditorSettings->InertiaDamping);
 	}
 	else
 	{
-		CVar->Set(0);
+		SetInertiaDampingCVar(0);
 	}
 	VPUtilitiesEditorSettings->SaveConfig();
 }
 
+void UVPScoutingSubsystem::SetInertiaDampingCVar(const float InInertiaDamping)
+{
+	IConsoleVariable* CVar = IConsoleManager::Get().FindConsoleVariable(TEXT("VI.HighSpeedInertiaDamping"));
+	CVar->Set(InInertiaDamping);
+}
+
 bool UVPScoutingSubsystem::IsHelperSystemEnabled()
 {
 	return GetDefault<UVPUtilitiesEditorSettings>()->bIsHelperSystemEnabled;
diff --git a/Engine/Plugins/Experimental/VirtualProductionUtilities/Source/VPUtilitiesEditor/Public/VPScoutingSubsystem.h b/Engine/Plugins/Experimental/VirtualProductionUtilities/Source/VPUtilitiesEditor/Public/VPScoutingSubsystem.h
--- a/Engine/Plugins/Experimental/VirtualProductionUtilities/Source/VPUtilitiesEditor/Public/VPScoutingSubsystem.h
+++ b/Engine/Plugins/Experimental/VirtualProductionUtilities/Source/VPUtilitiesEditor/Public/VPScoutingSubsystem.h
@@ -170,6 +170,10 @@ public:
 	UFUNCTION(BlueprintCallable, Category = "Virtual Production")
 	static void SetIsUsingTransformGizmo(const bool bInIsUsingTransformGizmo);
 
+	/** Set value of cvar "VI.ShowTransformGizmo" */
+	UFUNCTION(BlueprintCallable, Category = "Virtual Production")
+	static void SetShowTransformGizmoCVar(const bool bInShowTransformGizmoCVar);
+
 	/** Get flight speed for scouting in VR */
 	UFUNCTION(BlueprintPure, Category = "Virtual Production")
 	static float GetFlightSpeed();
@@ -194,6 +198,10 @@ public:
 	UFUNCTION(BlueprintCallable, Category = "Virtual Production")
 	static void SetIsUsingInertiaDamping(const bool bInIsUsingInertiaDamping);
 
+	/** Set value of cvar "VI.HighSpeedInertiaDamping" */
+	UFUNCTION(BlueprintCallable, Category = "Virtual Production")
+	static void SetInertiaDampingCVar(const float InInertiaDamping);
+
 	/** Whether the helper system on the controllers is enabled */
 	UFUNCTION(BlueprintPure, Category = "Virtual Production")
 	static bool IsHelperSystemEnabled();
