VPViewportTickableActorBase now tick only in the editor viewport

diff --git a/Engine/Plugins/Experimental/VirtualProductionUtilities/Source/VPUtilities/Private/VPViewportTickableActorBase.cpp b/Engine/Plugins/Experimental/VirtualProductionUtilities/Source/VPUtilities/Private/VPViewportTickableActorBase.cpp
--- a/Engine/Plugins/Experimental/VirtualProductionUtilities/Source/VPUtilities/Private/VPViewportTickableActorBase.cpp
+++ b/Engine/Plugins/Experimental/VirtualProductionUtilities/Source/VPUtilities/Private/VPViewportTickableActorBase.cpp
@@ -4,6 +4,7 @@
 
 
 AVPViewportTickableActorBase::AVPViewportTickableActorBase(const FObjectInitializer& ObjectInitializer) : Super(ObjectInitializer)
+	: ViewportTickType(EVPViewportTickableFlags::Game | EVPViewportTickableFlags::Editor)
 {
 	PrimaryActorTick.bCanEverTick = true;
 	PrimaryActorTick.bStartWithTickEnabled = true;
@@ -13,37 +14,62 @@ AVPViewportTickableActorBase::AVPViewportTickableActorBase(const FObjectInitiali
 }
 
 
- bool AVPViewportTickableActorBase::ShouldTickIfViewportsOnly() const
+bool AVPViewportTickableActorBase::ShouldTickIfViewportsOnly() const
 {
-	 return true;
+	if (UWorld* World = GetWorld())
+	{
+		switch (World->WorldType)
+		{
+		case EWorldType::Game:
+		case EWorldType::GameRPC:
+		case EWorldType::PIE:
+			return static_cast<uint8>(ViewportTickType & EVPViewportTickableFlags::Game) != 0;
+		case EWorldType::Editor:
+			return static_cast<uint8>(ViewportTickType & EVPViewportTickableFlags::Editor) != 0;
+		case EWorldType::Preview:
+		case EWorldType::GamePreview:
+			return static_cast<uint8>(ViewportTickType & EVPViewportTickableFlags::Preview) != 0;
+		}
+	}
+	return false;
 }
- 
 
- void AVPViewportTickableActorBase::Tick(float DeltaSeconds)
- {
-	 Super::Tick(DeltaSeconds);
 
-	 FEditorScriptExecutionGuard ScriptGuard;
-	 EditorTick(DeltaSeconds);
- }
+void AVPViewportTickableActorBase::Tick(float DeltaSeconds)
+{
+	Super::Tick(DeltaSeconds);
+
+#if WITH_EDITOR
+	if (GIsEditor)
+	{
+		FEditorScriptExecutionGuard ScriptGuard;
+		EditorTick(DeltaSeconds);
+	}
+#endif
+}
 
- void AVPViewportTickableActorBase::Destroyed()
- {
-	 FEditorScriptExecutionGuard ScriptGuard;
-	 EditorDestroyed();
 
-	 Super::Destroyed();
- }
+void AVPViewportTickableActorBase::Destroyed()
+{
+#if WITH_EDITOR
+	if (GIsEditor)
+	{
+		FEditorScriptExecutionGuard ScriptGuard;
+		EditorDestroyed();
+	}
+#endif
+
+	Super::Destroyed();
+}
 
 
- void AVPViewportTickableActorBase::EditorTick_Implementation(float DeltaSeconds)
- {
+void AVPViewportTickableActorBase::EditorTick_Implementation(float DeltaSeconds)
+{
 
- }
+}
 
- void AVPViewportTickableActorBase::EditorDestroyed_Implementation()
- {
 
- }
+void AVPViewportTickableActorBase::EditorDestroyed_Implementation()
+{
 
- 
\ No newline at end of file
+}
diff --git a/Engine/Plugins/Experimental/VirtualProductionUtilities/Source/VPUtilities/Public/VPViewportTickableActorBase.h b/Engine/Plugins/Experimental/VirtualProductionUtilities/Source/VPUtilities/Public/VPViewportTickableActorBase.h
--- a/Engine/Plugins/Experimental/VirtualProductionUtilities/Source/VPUtilities/Public/VPViewportTickableActorBase.h
+++ b/Engine/Plugins/Experimental/VirtualProductionUtilities/Source/VPUtilities/Public/VPViewportTickableActorBase.h
@@ -8,11 +8,20 @@
 
 #include "VPViewportTickableActorBase.generated.h"
 
+UENUM(meta = (Bitflags, UseEnumValuesAsMaskValuesInEditor = "true"))
+enum class EVPViewportTickableFlags : uint8
+{
+	Editor,
+	Game,
+	Preview,
+};
+ENUM_CLASS_FLAGS(EVPViewportTickableFlags)
+
 /**
  * Actor that tick in the Editor viewport with the event EditorTick.
  */
 UCLASS(Abstract)
-class AVPViewportTickableActorBase : public AActor
+class VPUTILITIES_API AVPViewportTickableActorBase : public AActor
 {
 	GENERATED_BODY()
 
@@ -25,10 +34,17 @@ public:
 	UFUNCTION(BlueprintNativeEvent, CallInEditor, BlueprintCallable, Category = "Utilities")
 	void EditorDestroyed();
 
+	/**
+	 * Where the actor should be ticked.
+	 * Editor = Tick in the editor viewport. Use the event EditorTick.
+	 * Game = Tick in game even if we are only ticking the viewport. Use the event Tick.
+	 * Preview = Tick if the actor is present in any editing tool like Blueprint or Material graph. Use EditorTick.
+	 */
+	UPROPERTY(EditDefaultsOnly, Category = "Actor Tick", meta = (Bitmask, BitmaskEnum = EVPViewportTickableType))
+	EVPViewportTickableFlags ViewportTickType;
 
 	/** If true, actor is ticked even if TickType==LEVELTICK_ViewportsOnly */
 	virtual bool ShouldTickIfViewportsOnly() const override;
 	virtual void Tick(float DeltaSeconds) override;
 	virtual void Destroyed() override;
-
 };
diff --git a/Engine/Plugins/Experimental/VirtualProductionUtilities/Source/VPUtilitiesEditor/Private/VPEditorTickableActorBase.cpp b/Engine/Plugins/Experimental/VirtualProductionUtilities/Source/VPUtilitiesEditor/Private/VPEditorTickableActorBase.cpp
--- a/Engine/Plugins/Experimental/VirtualProductionUtilities/Source/VPUtilitiesEditor/Private/VPEditorTickableActorBase.cpp
+++ b/Engine/Plugins/Experimental/VirtualProductionUtilities/Source/VPUtilitiesEditor/Private/VPEditorTickableActorBase.cpp
@@ -3,49 +3,6 @@
 #include "VPEditorTickableActorBase.h"
 
 
-AVPEditorTickableActorBase::AVPEditorTickableActorBase(const FObjectInitializer& ObjectInitializer) : Super(ObjectInitializer)
-{
-	PrimaryActorTick.bCanEverTick = true;
-	PrimaryActorTick.bStartWithTickEnabled = true;
-	SetActorTickEnabled(true);
-	//we don't want virtual production objects to be visible by cameras
-	SetActorHiddenInGame(true);
-}
-
-
- bool AVPEditorTickableActorBase::ShouldTickIfViewportsOnly() const
-{
-	 return true;
-}
- 
-
- void AVPEditorTickableActorBase::Tick(float DeltaSeconds)
- {
-	 Super::Tick(DeltaSeconds);
-
-	 FEditorScriptExecutionGuard ScriptGuard;
-	 EditorTick(DeltaSeconds);
- }
-
- void AVPEditorTickableActorBase::Destroyed()
- {
-	 FEditorScriptExecutionGuard ScriptGuard;
-	 EditorDestroyed();
-
-	 Super::Destroyed();
- }
-
-
- void AVPEditorTickableActorBase::EditorTick_Implementation(float DeltaSeconds)
- {
-
- }
-
- void AVPEditorTickableActorBase::EditorDestroyed_Implementation()
- {
-
- }
-
  void AVPEditorTickableActorBase::LockLocation(bool bSetLockLocation) 
  {
 	 bLockLocation = bSetLockLocation;
diff --git a/Engine/Plugins/Experimental/VirtualProductionUtilities/Source/VPUtilitiesEditor/Public/VPEditorTickableActorBase.h b/Engine/Plugins/Experimental/VirtualProductionUtilities/Source/VPUtilitiesEditor/Public/VPEditorTickableActorBase.h
--- a/Engine/Plugins/Experimental/VirtualProductionUtilities/Source/VPUtilitiesEditor/Public/VPEditorTickableActorBase.h
+++ b/Engine/Plugins/Experimental/VirtualProductionUtilities/Source/VPUtilitiesEditor/Public/VPEditorTickableActorBase.h
@@ -3,8 +3,7 @@
 #pragma once
 
 #include "CoreMinimal.h"
-#include "GameFramework/Actor.h"
-#include "UObject/UObjectGlobals.h"
+#include "VPViewportTickableActorBase.h"
 
 #include "VPEditorTickableActorBase.generated.h"
 
@@ -12,24 +11,11 @@
  * Actor that tick in the Editor viewport with the event EditorTick.
  */
 UCLASS(Abstract)
-class AVPEditorTickableActorBase : public AActor
+class VPUTILITIESEDITOR_API AVPEditorTickableActorBase : public AVPViewportTickableActorBase
 {
 	GENERATED_BODY()
 
 public:
-	AVPEditorTickableActorBase(const FObjectInitializer& ObjectInitializer);
-
-	UFUNCTION(BlueprintNativeEvent, CallInEditor, BlueprintCallable, Category = "Tick")
-	void EditorTick(float DeltaSeconds);
-
-	UFUNCTION(BlueprintNativeEvent, CallInEditor, BlueprintCallable, Category = "Utilities")
-	void EditorDestroyed();
-
-
-	/** If true, actor is ticked even if TickType==LEVELTICK_ViewportsOnly */
-	virtual bool ShouldTickIfViewportsOnly() const override;
-	virtual void Tick(float DeltaSeconds) override;
-	virtual void Destroyed() override;
 
 	/** Sets the LockLocation variable to disable movement from the translation gizmo */
 	UFUNCTION(CallInEditor, BlueprintCallable, Category = "Editor")
