Fix bad auto-merge introduced with CL 10360685

diff --git a/Engine/Plugins/Experimental/VirtualProductionUtilities/Source/VPUtilities/Private/VPViewportTickableActorBase.cpp b/Engine/Plugins/Experimental/VirtualProductionUtilities/Source/VPUtilities/Private/VPViewportTickableActorBase.cpp
--- a/Engine/Plugins/Experimental/VirtualProductionUtilities/Source/VPUtilities/Private/VPViewportTickableActorBase.cpp
+++ b/Engine/Plugins/Experimental/VirtualProductionUtilities/Source/VPUtilities/Private/VPViewportTickableActorBase.cpp
@@ -3,8 +3,9 @@
 #include "VPViewportTickableActorBase.h"
 
 
-AVPViewportTickableActorBase::AVPViewportTickableActorBase(const FObjectInitializer& ObjectInitializer) : Super(ObjectInitializer)
-	: ViewportTickType(EVPViewportTickableFlags::Game | EVPViewportTickableFlags::Editor)
+AVPViewportTickableActorBase::AVPViewportTickableActorBase(const FObjectInitializer& ObjectInitializer)
+	: Super(ObjectInitializer)
+	, ViewportTickType(EVPViewportTickableFlags::Game | EVPViewportTickableFlags::Editor)
 {
 	PrimaryActorTick.bCanEverTick = true;
 	PrimaryActorTick.bStartWithTickEnabled = true;
