Fix CIS. Missing include. Introduced with CL 10360685.

diff --git a/Engine/Plugins/Experimental/VirtualProductionUtilities/Source/VPUtilities/Private/VPViewportTickableActorBase.cpp b/Engine/Plugins/Experimental/VirtualProductionUtilities/Source/VPUtilities/Private/VPViewportTickableActorBase.cpp
--- a/Engine/Plugins/Experimental/VirtualProductionUtilities/Source/VPUtilities/Private/VPViewportTickableActorBase.cpp
+++ b/Engine/Plugins/Experimental/VirtualProductionUtilities/Source/VPUtilities/Private/VPViewportTickableActorBase.cpp
@@ -2,6 +2,8 @@
 
 #include "VPViewportTickableActorBase.h"
 
+#include "Engine/World.h"
+
 
 AVPViewportTickableActorBase::AVPViewportTickableActorBase(const FObjectInitializer& ObjectInitializer)
 	: Super(ObjectInitializer)
