Fix CIS warning about deprecated variable in VirtualProductionUtilities.

diff --git a/Engine/Plugins/Experimental/VirtualProductionUtilities/Source/VPUtilities/Private/VPViewportTickableActorBase.cpp b/Engine/Plugins/Experimental/VirtualProductionUtilities/Source/VPUtilities/Private/VPViewportTickableActorBase.cpp
--- a/Engine/Plugins/Experimental/VirtualProductionUtilities/Source/VPUtilities/Private/VPViewportTickableActorBase.cpp
+++ b/Engine/Plugins/Experimental/VirtualProductionUtilities/Source/VPUtilities/Private/VPViewportTickableActorBase.cpp
@@ -29,9 +29,10 @@ bool AVPViewportTickableActorBase::ShouldTickIfViewportsOnly() const
 			return static_cast<uint8>(ViewportTickType & EVPViewportTickableFlags::Game) != 0;
 		case EWorldType::Editor:
 			return static_cast<uint8>(ViewportTickType & EVPViewportTickableFlags::Editor) != 0;
-		case EWorldType::Preview:
+		case EWorldType::EditorPreview:
+			return static_cast<uint8>(ViewportTickType & EVPViewportTickableFlags::EditorPreview) != 0;
 		case EWorldType::GamePreview:
-			return static_cast<uint8>(ViewportTickType & EVPViewportTickableFlags::Preview) != 0;
+			return static_cast<uint8>(ViewportTickType & EVPViewportTickableFlags::GamePreview) != 0;
 		}
 	}
 	return false;
diff --git a/Engine/Plugins/Experimental/VirtualProductionUtilities/Source/VPUtilities/Public/VPViewportTickableActorBase.h b/Engine/Plugins/Experimental/VirtualProductionUtilities/Source/VPUtilities/Public/VPViewportTickableActorBase.h
--- a/Engine/Plugins/Experimental/VirtualProductionUtilities/Source/VPUtilities/Public/VPViewportTickableActorBase.h
+++ b/Engine/Plugins/Experimental/VirtualProductionUtilities/Source/VPUtilities/Public/VPViewportTickableActorBase.h
@@ -13,7 +13,8 @@ enum class EVPViewportTickableFlags : uint8
 {
 	Editor,
 	Game,
-	Preview,
+	EditorPreview,
+	GamePreview,
 };
 ENUM_CLASS_FLAGS(EVPViewportTickableFlags)
 
