Initializing the VPViewportTickableFlags properly.
#rnx

diff --git a/Engine/Plugins/Experimental/VirtualProductionUtilities/Source/VPUtilities/Public/VPViewportTickableActorBase.h b/Engine/Plugins/Experimental/VirtualProductionUtilities/Source/VPUtilities/Public/VPViewportTickableActorBase.h
--- a/Engine/Plugins/Experimental/VirtualProductionUtilities/Source/VPUtilities/Public/VPViewportTickableActorBase.h
+++ b/Engine/Plugins/Experimental/VirtualProductionUtilities/Source/VPUtilities/Public/VPViewportTickableActorBase.h
@@ -11,10 +11,10 @@
 UENUM(meta = (Bitflags, UseEnumValuesAsMaskValuesInEditor = "true"))
 enum class EVPViewportTickableFlags : uint8
 {
-	Editor,
-	Game,
-	EditorPreview,
-	GamePreview,
+	Editor			= 1 << 0,
+	Game			= 1 << 1,
+	EditorPreview	= 1 << 2,
+	GamePreview		= 1 << 3,
 };
 ENUM_CLASS_FLAGS(EVPViewportTickableFlags)
 
