Create an actor that control how to display a UMG in fullscreen in different game mode. Currently, it could be in PostProcess or in Viewport's widget. Could be usefull to display and capture the widget in editor.

diff --git a/Engine/Plugins/Experimental/VirtualProductionUtilities/Source/VPUtilities/Private/VPFullScreenUserWidget.cpp b/Engine/Plugins/Experimental/VirtualProductionUtilities/Source/VPUtilities/Private/VPFullScreenUserWidget.cpp
new file mode 100644
--- /dev/null
+++ b/Engine/Plugins/Experimental/VirtualProductionUtilities/Source/VPUtilities/Private/VPFullScreenUserWidget.cpp
@@ -0,0 +1,717 @@
+// Copyright 1998-2019 Epic Games, Inc. All Rights Reserved.
+
+#include "VPFullScreenUserWidget.h"
+
+#include "Components/PostProcessComponent.h"
+#include "Engine/Texture2D.h"
+#include "Engine/TextureRenderTarget2D.h"
+#include "Engine/UserInterfaceSettings.h"
+#include "GameFramework/WorldSettings.h"
+#include "Materials/MaterialInstanceDynamic.h"
+#include "RenderingThread.h"
+#include "RHI.h"
+#include "UObject/ConstructorHelpers.h"
+#include "VPUtilitiesModule.h"
+
+#include "Framework/Application/SlateApplication.h"
+#include "Input/HittestGrid.h"
+#include "Layout/Visibility.h"
+#include "Slate/SceneViewport.h"
+#include "Slate/WidgetRenderer.h"
+#include "Widgets/Layout/SConstraintCanvas.h"
+#include "Widgets/SViewport.h"
+
+
+#if WITH_EDITOR
+#include "LevelEditor.h"
+#include "Modules/ModuleManager.h"
+#include "SLevelViewport.h"
+#endif
+
+#define LOCTEXT_NAMESPACE "VPFullScreenUserWidget"
+
+/////////////////////////////////////////////////////
+// Internal helper
+namespace
+{
+	const FName NAME_LevelEditorName = "LevelEditor";
+	const FName NAME_SlateUI = "SlateUI";
+	const FName NAME_TintColorAndOpacity = "TintColorAndOpacity";
+	const FName NAME_OpacityFromTexture = "OpacityFromTexture";
+
+	EVisibility ConvertWindowVisibilityToVisibility(EWindowVisibility visibility)
+	{
+		switch (visibility)
+		{
+		case EWindowVisibility::Visible:
+			return EVisibility::Visible;
+		case EWindowVisibility::SelfHitTestInvisible:
+			return EVisibility::SelfHitTestInvisible;
+		default:
+			checkNoEntry();
+			return EVisibility::SelfHitTestInvisible;
+		}
+	}
+}
+
+
+/////////////////////////////////////////////////////
+// FVPWidgetPostProcessHitTester
+class FVPWidgetPostProcessHitTester : public ICustomHitTestPath
+{
+public:
+	FVPWidgetPostProcessHitTester(UWorld* InWorld, TSharedPtr<SVirtualWindow> InSlateWindow)
+		: World(InWorld)
+		, SlateWindow(InSlateWindow)
+		, WidgetDrawSize(FIntPoint::ZeroValue)
+		, LastLocalHitLocation(FVector2D::ZeroVector)
+	{}
+
+	virtual TArray<FWidgetAndPointer> GetBubblePathAndVirtualCursors(const FGeometry& InGeometry, FVector2D DesktopSpaceCoordinate, bool bIgnoreEnabledStatus) const override
+	{
+		// Get the list of widget at the requested location.
+		TArray<FWidgetAndPointer> ArrangedWidgets;
+		if (TSharedPtr<SVirtualWindow> SlateWindowPin = SlateWindow.Pin())
+		{
+			FVector2D LocalMouseCoordinate = InGeometry.AbsoluteToLocal(DesktopSpaceCoordinate) * InGeometry.Scale;
+			float CursorRadius = 0.f;
+			ArrangedWidgets = SlateWindowPin->GetHittestGrid().GetBubblePath(LocalMouseCoordinate, CursorRadius, bIgnoreEnabledStatus);
+
+			TSharedRef<FVirtualPointerPosition> VirtualMouseCoordinate = MakeShared<FVirtualPointerPosition>();
+			VirtualMouseCoordinate->CurrentCursorPosition = LocalMouseCoordinate;
+			VirtualMouseCoordinate->LastCursorPosition = LastLocalHitLocation;
+
+			LastLocalHitLocation = LocalMouseCoordinate;
+
+			for (FWidgetAndPointer& ArrangedWidget : ArrangedWidgets)
+			{
+				ArrangedWidget.PointerPosition = VirtualMouseCoordinate;
+			}
+		}
+
+		return ArrangedWidgets;
+	}
+
+	virtual void ArrangeChildren(FArrangedChildren& ArrangedChildren) const override
+	{
+		// Add the displayed slate to the list of widgets.
+		if (TSharedPtr<SVirtualWindow> SlateWindowPin = SlateWindow.Pin())
+		{
+			FGeometry WidgetGeom;
+			ArrangedChildren.AddWidget(FArrangedWidget(SlateWindowPin.ToSharedRef(), WidgetGeom.MakeChild(WidgetDrawSize, FSlateLayoutTransform())));
+		}
+	}
+
+	virtual TSharedPtr<struct FVirtualPointerPosition> TranslateMouseCoordinateFor3DChild(const TSharedRef<SWidget>& ChildWidget, const FGeometry& ViewportGeometry, const FVector2D& ScreenSpaceMouseCoordinate, const FVector2D& LastScreenSpaceMouseCoordinate) const override
+	{
+		return nullptr;
+	}
+
+	void SetWidgetDrawSize(FIntPoint NewWidgetDrawSize)
+	{
+		WidgetDrawSize = NewWidgetDrawSize;
+	}
+
+private:
+	TWeakObjectPtr<UWorld> World;
+	TWeakPtr<SVirtualWindow> SlateWindow;
+	FIntPoint WidgetDrawSize;
+	mutable FVector2D LastLocalHitLocation;
+};
+
+/////////////////////////////////////////////////////
+// FVPFullScreenUserWidget_Viewport
+FVPFullScreenUserWidget_Viewport::FVPFullScreenUserWidget_Viewport()
+	: bAddedToGameViewport(false)
+{
+}
+
+bool FVPFullScreenUserWidget_Viewport::Display(UWorld* World, UUserWidget* Widget)
+{
+	TSharedPtr<SConstraintCanvas> FullScreenWidgetPinned = FullScreenCanvasWidget.Pin();
+	if (Widget == nullptr || World == nullptr || FullScreenWidgetPinned.IsValid())
+	{
+		return false;
+	}
+
+	UGameViewportClient* ViewportClient = nullptr;
+#if WITH_EDITOR
+	TSharedPtr<SLevelViewport> ActiveLevelViewport;
+#endif
+
+	bool bResult = false;
+	if (World->WorldType == EWorldType::Game || World->WorldType == EWorldType::PIE)
+	{
+		ViewportClient = World->GetGameViewport();
+		bResult = ViewportClient != nullptr;
+	}
+#if WITH_EDITOR
+	else if (FModuleManager::Get().IsModuleLoaded(NAME_LevelEditorName))
+	{
+		FLevelEditorModule& LevelEditorModule = FModuleManager::GetModuleChecked<FLevelEditorModule>(NAME_LevelEditorName);
+		ActiveLevelViewport = LevelEditorModule.GetFirstActiveLevelViewport();
+		bResult = ActiveLevelViewport.IsValid();
+	}
+#endif
+
+	if (bResult)
+	{
+		TSharedRef<SConstraintCanvas> FullScreenCanvas = SNew(SConstraintCanvas);
+		FullScreenCanvasWidget = FullScreenCanvas;
+
+		FullScreenCanvas->AddSlot()
+			.Offset(FMargin(0, 0, 0, 0))
+			.Anchors(FAnchors(0, 0, 1, 1))
+			.Alignment(FVector2D(0, 0))
+			[
+				Widget->TakeWidget()
+			];
+
+		if (ViewportClient)
+		{
+			ViewportClient->AddViewportWidgetContent(FullScreenCanvas);
+		}
+#if WITH_EDITOR
+		else
+		{
+			check(ActiveLevelViewport.IsValid());
+			ActiveLevelViewport->AddOverlayWidget(FullScreenCanvas);
+			OverlayWidgetLevelViewport = ActiveLevelViewport;
+		}
+#endif
+	}
+
+	return bResult;
+}
+
+void FVPFullScreenUserWidget_Viewport::Hide(UWorld* World)
+{
+	TSharedPtr<SConstraintCanvas> FullScreenWidgetPinned = FullScreenCanvasWidget.Pin();
+	if (FullScreenWidgetPinned.IsValid())
+	{
+		// Remove from Viewport and Fullscreen, in case the settings changed before we had the chance to hide.
+		UGameViewportClient* ViewportClient = World ? World->GetGameViewport() : nullptr;
+		if (ViewportClient)
+		{
+			ViewportClient->RemoveViewportWidgetContent(FullScreenWidgetPinned.ToSharedRef());
+		}
+
+#if WITH_EDITOR
+		TSharedPtr<SLevelViewport> OverlayWidgetLevelViewportPinned = OverlayWidgetLevelViewport.Pin();
+		if (OverlayWidgetLevelViewportPinned)
+		{
+			OverlayWidgetLevelViewportPinned->RemoveOverlayWidget(FullScreenWidgetPinned.ToSharedRef());
+		}
+		OverlayWidgetLevelViewport.Reset();
+#endif
+
+		FullScreenCanvasWidget.Reset();
+	}
+}
+
+void FVPFullScreenUserWidget_Viewport::Tick(UWorld* World, float DeltaSeconds)
+{
+
+}
+
+/////////////////////////////////////////////////////
+// FVPFullScreenUserWidget_PostProcess
+
+FVPFullScreenUserWidget_PostProcess::FVPFullScreenUserWidget_PostProcess()
+	: PostProcessMaterial(nullptr)
+	, PostProcessTintColorAndOpacity(FLinearColor::White)
+	, PostProcessOpacityFromTexture(1.0f)
+	, WidgetDrawSize(FIntPoint(640, 360))
+	, bWindowFocusable(true)
+	, WindowVisibility(EWindowVisibility::SelfHitTestInvisible)
+	, bReceiveHardwareInput(false)
+	, RenderTargetBackgroundColor(FLinearColor(0.0f, 0.0f, 0.0f, 1.0f))
+	, RenderTargetBlendMode(EWidgetBlendMode::Masked)
+	, PostProcessComponent(nullptr)
+	, PostProcessMaterialInstance(nullptr)
+	, WidgetRenderTarget(nullptr)
+	, WidgetRenderer(nullptr)
+	, CurrentWidgetDrawSize(FIntPoint::ZeroValue)
+{
+}
+
+bool FVPFullScreenUserWidget_PostProcess::Display(UWorld* World, UUserWidget* Widget)
+{
+	return CreateRenderer(World, Widget) && CreatePostProcessComponent(World);
+}
+
+void FVPFullScreenUserWidget_PostProcess::Hide(UWorld* World)
+{
+	ReleasePostProcessComponent();
+	ReleaseRenderer();
+}
+
+void FVPFullScreenUserWidget_PostProcess::Tick(UWorld* World, float DeltaSeconds)
+{
+	TickRenderer(World, DeltaSeconds);
+}
+
+bool FVPFullScreenUserWidget_PostProcess::CreatePostProcessComponent(UWorld* World)
+{
+	ReleasePostProcessComponent();
+	if (World && PostProcessMaterial)
+	{
+		AWorldSettings* WorldSetting = World->GetWorldSettings();
+		PostProcessComponent = NewObject<UPostProcessComponent>(WorldSetting, NAME_None, RF_Transient);
+		PostProcessComponent->bEnabled = true;
+		PostProcessComponent->bUnbound = true;
+		PostProcessComponent->RegisterComponent();
+
+		PostProcessMaterialInstance = UMaterialInstanceDynamic::Create(PostProcessMaterial, World);
+
+		// set the parameter immediately
+		PostProcessMaterialInstance->SetTextureParameterValue(NAME_SlateUI, WidgetRenderTarget);
+		PostProcessMaterialInstance->SetVectorParameterValue(NAME_TintColorAndOpacity, PostProcessTintColorAndOpacity);
+		PostProcessMaterialInstance->SetScalarParameterValue(NAME_OpacityFromTexture, PostProcessOpacityFromTexture);
+
+		PostProcessComponent->Settings.WeightedBlendables.Array.SetNumZeroed(1);
+		PostProcessComponent->Settings.WeightedBlendables.Array[0].Weight = 1.f;
+		PostProcessComponent->Settings.WeightedBlendables.Array[0].Object = PostProcessMaterialInstance;
+	}
+
+	return PostProcessComponent && PostProcessMaterialInstance;
+}
+
+void FVPFullScreenUserWidget_PostProcess::ReleasePostProcessComponent()
+{
+	if (PostProcessComponent)
+	{
+		PostProcessComponent->UnregisterComponent();
+	}
+	PostProcessComponent = nullptr;
+	PostProcessMaterialInstance = nullptr;
+}
+
+bool FVPFullScreenUserWidget_PostProcess::CreateRenderer(UWorld* World, UUserWidget* Widget)
+{
+	ReleaseRenderer();
+
+	if (World && Widget)
+	{
+		const FIntPoint CalculatedWidgetSize = CalculateWidgetDrawSize(World);
+		if (IsTextureSizeValid(CalculatedWidgetSize))
+		{
+			CurrentWidgetDrawSize = CalculatedWidgetSize;
+
+			const bool bApplyGammaCorrection = false;
+			WidgetRenderer = new FWidgetRenderer(bApplyGammaCorrection);
+			WidgetRenderer->SetIsPrepassNeeded(true);
+
+			SlateWindow = SNew(SVirtualWindow).Size(CurrentWidgetDrawSize);
+			SlateWindow->SetIsFocusable(bWindowFocusable);
+			SlateWindow->SetVisibility(ConvertWindowVisibilityToVisibility(WindowVisibility));
+			SlateWindow->SetContent(Widget->TakeWidget());
+
+			RegisterHitTesterWithViewport(World);
+
+			if (!Widget->IsDesignTime() && World->IsGameWorld())
+			{
+				UGameInstance* GameInstance = World->GetGameInstance();
+				UGameViewportClient* GameViewportClient = GameInstance ? GameInstance->GetGameViewportClient() : nullptr;
+				if (GameViewportClient)
+				{
+					SlateWindow->AssignParentWidget(GameViewportClient->GetGameViewportWidget());
+				}
+			}
+
+			FLinearColor ActualBackgroundColor = RenderTargetBackgroundColor;
+			switch (RenderTargetBlendMode)
+			{
+			case EWidgetBlendMode::Opaque:
+				ActualBackgroundColor.A = 1.0f;
+				break;
+			case EWidgetBlendMode::Masked:
+				ActualBackgroundColor.A = 0.0f;
+				break;
+			}
+
+			AWorldSettings* WorldSetting = World->GetWorldSettings();
+			WidgetRenderTarget = NewObject<UTextureRenderTarget2D>(WorldSetting, NAME_None, RF_Transient);
+			WidgetRenderTarget->ClearColor = ActualBackgroundColor;
+			WidgetRenderTarget->InitCustomFormat(CurrentWidgetDrawSize.X, CurrentWidgetDrawSize.Y, PF_B8G8R8A8, false);
+			WidgetRenderTarget->UpdateResourceImmediate();
+
+			if (PostProcessMaterialInstance)
+			{
+				PostProcessMaterialInstance->SetTextureParameterValue(NAME_SlateUI, WidgetRenderTarget);
+			}
+
+		}
+	}
+
+	return WidgetRenderer && WidgetRenderTarget;
+}
+
+void FVPFullScreenUserWidget_PostProcess::ReleaseRenderer()
+{
+	if (WidgetRenderer)
+	{
+		BeginCleanup(WidgetRenderer);
+		WidgetRenderer = nullptr;
+	}
+	UnRegisterHitTesterWithViewport();
+
+	SlateWindow.Reset();
+	WidgetRenderTarget = nullptr;
+	CurrentWidgetDrawSize = FIntPoint::ZeroValue;
+}
+
+void FVPFullScreenUserWidget_PostProcess::TickRenderer(UWorld* World, float DeltaSeconds)
+{
+	check(World);
+	if (WidgetRenderTarget)
+	{
+		const float DrawScale = 1.0f;
+
+		const FIntPoint NewCalculatedWidgetSize = CalculateWidgetDrawSize(World);
+		if (NewCalculatedWidgetSize != CurrentWidgetDrawSize)
+		{
+			if (IsTextureSizeValid(NewCalculatedWidgetSize))
+			{
+				CurrentWidgetDrawSize = NewCalculatedWidgetSize;
+				WidgetRenderTarget->InitCustomFormat(CurrentWidgetDrawSize.X, CurrentWidgetDrawSize.Y, PF_B8G8R8A8, false);
+				WidgetRenderTarget->UpdateResourceImmediate();
+				SlateWindow->Resize(CurrentWidgetDrawSize);
+				if (CustomHitTestPath)
+				{
+					CustomHitTestPath->SetWidgetDrawSize(CurrentWidgetDrawSize);
+				}
+			}
+			else
+			{
+				Hide(World);
+			}
+		}
+
+		if (WidgetRenderer)
+		{
+			WidgetRenderer->DrawWindow(
+				WidgetRenderTarget,
+				SlateWindow->GetHittestGrid(),
+				SlateWindow.ToSharedRef(),
+				DrawScale,
+				CurrentWidgetDrawSize,
+				DeltaSeconds);
+		}
+	}
+}
+
+FIntPoint FVPFullScreenUserWidget_PostProcess::CalculateWidgetDrawSize(UWorld* World)
+{
+	if (bWidgetDrawSize)
+	{
+		return WidgetDrawSize;
+	}
+
+	if (World->WorldType == EWorldType::Game || World->WorldType == EWorldType::PIE)
+	{
+		if (UGameViewportClient* ViewportClient = World->GetGameViewport())
+		{
+			// The viewport maybe resizing or not yet initialized.
+			//See TickRenderer(), it will be resize on the next tick to the proper size.
+			//We initialized all the rendering with an small size.
+
+			const float SmallWidgetSize = 16.f;
+			FVector2D OutSize = FVector2D(SmallWidgetSize, SmallWidgetSize);
+			ViewportClient->GetViewportSize(OutSize);
+			if (OutSize.X < SMALL_NUMBER)
+			{
+				OutSize = FVector2D(SmallWidgetSize, SmallWidgetSize);
+			}
+			return OutSize.IntPoint();
+		}
+	}
+#if WITH_EDITOR
+	else if (FModuleManager::Get().IsModuleLoaded(NAME_LevelEditorName))
+	{
+		FLevelEditorModule& LevelEditorModule = FModuleManager::GetModuleChecked<FLevelEditorModule>(NAME_LevelEditorName);
+		if (TSharedPtr<SLevelViewport> ActiveLevelViewport = LevelEditorModule.GetFirstActiveLevelViewport())
+		{
+			if (TSharedPtr<FSceneViewport> SharedActiveViewport = ActiveLevelViewport->GetSharedActiveViewport())
+			{
+				return SharedActiveViewport->GetSize();
+			}
+		}
+	}
+#endif
+	return FIntPoint::ZeroValue;
+}
+
+bool FVPFullScreenUserWidget_PostProcess::IsTextureSizeValid(FIntPoint Size) const
+{
+	const int32 MaxAllowedDrawSize = GetMax2DTextureDimension();
+	return Size.X > 0 && Size.Y > 0 && Size.X <= MaxAllowedDrawSize && Size.Y <= MaxAllowedDrawSize;
+}
+
+void FVPFullScreenUserWidget_PostProcess::RegisterHitTesterWithViewport(UWorld* World)
+{
+	if (!bReceiveHardwareInput && FSlateApplication::IsInitialized())
+	{
+		FSlateApplication::Get().RegisterVirtualWindow(SlateWindow.ToSharedRef());
+	}
+
+	TSharedPtr<SViewport> EngineViewportWidget;
+	if (World->WorldType == EWorldType::Game || World->WorldType == EWorldType::PIE)
+	{
+		EngineViewportWidget = GEngine->GetGameViewportWidget();
+	}
+#if WITH_EDITOR
+	else if (FModuleManager::Get().IsModuleLoaded(NAME_LevelEditorName))
+	{
+		FLevelEditorModule& LevelEditorModule = FModuleManager::GetModuleChecked<FLevelEditorModule>(NAME_LevelEditorName);
+		if (TSharedPtr<SLevelViewport> ActiveLevelViewport = LevelEditorModule.GetFirstActiveLevelViewport())
+		{
+			EngineViewportWidget = ActiveLevelViewport->GetViewportWidget().Pin();
+		}
+	}
+#endif
+
+	if (EngineViewportWidget && bReceiveHardwareInput)
+	{
+		if (EngineViewportWidget->GetCustomHitTestPath())
+		{
+			UE_LOG(LogVPUtilities, Warning, TEXT("Can't register a hit tester for FullScreenUserWidget. There is already one defined."));
+		}
+		else
+		{
+			ViewportWidget = EngineViewportWidget;
+			CustomHitTestPath = MakeShared<FVPWidgetPostProcessHitTester>(World, SlateWindow);
+			CustomHitTestPath->SetWidgetDrawSize(CurrentWidgetDrawSize);
+			EngineViewportWidget->SetCustomHitTestPath(CustomHitTestPath);
+		}
+	}
+}
+
+void FVPFullScreenUserWidget_PostProcess::UnRegisterHitTesterWithViewport()
+{
+	if (SlateWindow.IsValid() && FSlateApplication::IsInitialized())
+	{
+		FSlateApplication::Get().UnregisterVirtualWindow(SlateWindow.ToSharedRef());
+	}
+
+	if (TSharedPtr<SViewport> ViewportWidgetPin = ViewportWidget.Pin())
+	{
+		if (ViewportWidgetPin->GetCustomHitTestPath() == CustomHitTestPath)
+		{
+			ViewportWidgetPin->SetCustomHitTestPath(nullptr);
+		}
+	}
+
+	ViewportWidget.Reset();
+	CustomHitTestPath.Reset();
+}
+
+/////////////////////////////////////////////////////
+// UVPFullScreenUserWidget
+
+UVPFullScreenUserWidget::UVPFullScreenUserWidget(const FObjectInitializer& ObjectInitializer)
+	: Super(ObjectInitializer)
+	, CurrentDisplayType(EVPWidgetDisplayType::Inactive)
+	, bDisplayRequested(false)
+{
+	static ConstructorHelpers::FObjectFinder<UMaterialInterface> PostProcessMaterial_Finder(TEXT("/VirtualProductionUtilities/Materials/WidgetPostProcessMaterial"));
+	PostProcessDisplayType.PostProcessMaterial = PostProcessMaterial_Finder.Object;
+}
+
+void UVPFullScreenUserWidget::BeginDestroy()
+{
+	Hide();
+	Super::BeginDestroy();
+}
+
+bool UVPFullScreenUserWidget::ShouldDisplay(UWorld* InWorld) const
+{
+#if UE_SERVER
+	return false;
+#else
+	if (GUsingNullRHI || HasAnyFlags(RF_ArchetypeObject | RF_ClassDefaultObject) || IsRunningDedicatedServer())
+	{
+		return false;
+	}
+
+	return GetDisplayType(InWorld) != EVPWidgetDisplayType::Inactive;
+#endif //!UE_SERVER
+}
+
+EVPWidgetDisplayType UVPFullScreenUserWidget::GetDisplayType(UWorld* InWorld) const
+{
+	if (InWorld)
+	{
+		if (InWorld->WorldType == EWorldType::Game)
+		{
+			return GameDisplayType;
+		}
+#if WITH_EDITOR
+		else if (InWorld->WorldType == EWorldType::PIE)
+		{
+			return PIEDisplayType;
+		}
+		else if (InWorld->WorldType == EWorldType::Editor)
+		{
+			return EditorDisplayType;
+		}
+#endif // WITH_EDITOR
+	}
+	return EVPWidgetDisplayType::Inactive;
+}
+
+bool UVPFullScreenUserWidget::IsDisplayed() const
+{
+	return CurrentDisplayType != EVPWidgetDisplayType::Inactive;
+}
+
+bool UVPFullScreenUserWidget::Display(UWorld* InWorld)
+{
+	bDisplayRequested = true;
+
+	World = InWorld;
+
+	bool bWasAdded = false;
+	if (InWorld && WidgetClass && ShouldDisplay(InWorld) && CurrentDisplayType == EVPWidgetDisplayType::Inactive)
+	{
+		CurrentDisplayType = GetDisplayType(InWorld);
+
+		InitWidget();
+
+		if (CurrentDisplayType == EVPWidgetDisplayType::Viewport)
+		{
+			bWasAdded = ViewportDisplayType.Display(InWorld, Widget);
+		}
+		else if (CurrentDisplayType == EVPWidgetDisplayType::PostProcess)
+		{
+			bWasAdded = PostProcessDisplayType.Display(InWorld, Widget);
+		}
+
+		if (bWasAdded)
+		{
+			FWorldDelegates::LevelRemovedFromWorld.AddUObject(this, &UVPFullScreenUserWidget::OnLevelRemovedFromWorld);
+		}
+	}
+
+	return bWasAdded;
+}
+
+void UVPFullScreenUserWidget::Hide()
+{
+	bDisplayRequested = false;
+
+	if (CurrentDisplayType != EVPWidgetDisplayType::Inactive)
+	{
+		ReleaseWidget();
+		FWorldDelegates::LevelRemovedFromWorld.RemoveAll(this);
+
+		if (CurrentDisplayType == EVPWidgetDisplayType::Viewport)
+		{
+			ViewportDisplayType.Hide(World.Get());
+		}
+		else if (CurrentDisplayType == EVPWidgetDisplayType::PostProcess)
+		{
+			PostProcessDisplayType.Hide(World.Get());
+		}
+		CurrentDisplayType = EVPWidgetDisplayType::Inactive;
+	}
+
+	World.Reset();
+}
+
+void UVPFullScreenUserWidget::Tick(float DeltaSeconds)
+{
+	if (CurrentDisplayType != EVPWidgetDisplayType::Inactive)
+	{
+		UWorld* CurrentWorld = World.Get();
+		if (CurrentWorld == nullptr)
+		{
+			Hide();
+		}
+		else
+		{
+			if (CurrentDisplayType == EVPWidgetDisplayType::Viewport)
+			{
+				ViewportDisplayType.Tick(CurrentWorld, DeltaSeconds);
+			}
+			else if (CurrentDisplayType == EVPWidgetDisplayType::PostProcess)
+			{
+				PostProcessDisplayType.Tick(CurrentWorld, DeltaSeconds);
+			}
+		}
+	}
+}
+
+void UVPFullScreenUserWidget::InitWidget()
+{
+	// Don't do any work if Slate is not initialized
+	if (FSlateApplication::IsInitialized())
+	{
+		if (WidgetClass && Widget == nullptr)
+		{
+			check(GetWorld());
+			Widget = CreateWidget(GetWorld(), WidgetClass);
+			Widget->SetFlags(RF_Transient);
+		}
+	}
+}
+
+void UVPFullScreenUserWidget::ReleaseWidget()
+{
+	Widget = nullptr;
+}
+
+void UVPFullScreenUserWidget::OnLevelRemovedFromWorld(ULevel* InLevel, UWorld* InWorld)
+{
+	// If the InLevel is invalid, then the entire world is about to disappear.
+	//Hide the widget to clear the memory and reference to the world it may hold.
+	if (InLevel == nullptr && InWorld && InWorld == World.Get())
+	{
+		Hide();
+	}
+}
+
+#if WITH_EDITOR
+void UVPFullScreenUserWidget::PostEditChangeProperty(FPropertyChangedEvent& PropertyChangedEvent)
+{
+	UProperty* Property = PropertyChangedEvent.MemberProperty;
+
+	if (Property && PropertyChangedEvent.ChangeType != EPropertyChangeType::Interactive)
+	{
+		static FName NAME_WidgetClass = GET_MEMBER_NAME_CHECKED(UVPFullScreenUserWidget, WidgetClass);
+		static FName NAME_EditorDisplayType = GET_MEMBER_NAME_CHECKED(UVPFullScreenUserWidget, EditorDisplayType);
+		static FName NAME_PostProcessMaterial = GET_MEMBER_NAME_CHECKED(FVPFullScreenUserWidget_PostProcess, PostProcessMaterial);
+		static FName NAME_WidgetDrawSize = GET_MEMBER_NAME_CHECKED(FVPFullScreenUserWidget_PostProcess, WidgetDrawSize);
+		static FName NAME_WindowFocusable = GET_MEMBER_NAME_CHECKED(FVPFullScreenUserWidget_PostProcess, bWindowFocusable);
+		static FName NAME_WindowVisibility = GET_MEMBER_NAME_CHECKED(FVPFullScreenUserWidget_PostProcess, WindowVisibility);
+		static FName NAME_ReceiveHardwareInput = GET_MEMBER_NAME_CHECKED(FVPFullScreenUserWidget_PostProcess, bReceiveHardwareInput);
+		static FName NAME_RenderTargetBackgroundColor = GET_MEMBER_NAME_CHECKED(FVPFullScreenUserWidget_PostProcess, RenderTargetBackgroundColor);
+		static FName NAME_RenderTargetBlendMode = GET_MEMBER_NAME_CHECKED(FVPFullScreenUserWidget_PostProcess, RenderTargetBlendMode);
+		static FName NAME_PostProcessTintColorAndOpacity = GET_MEMBER_NAME_CHECKED(FVPFullScreenUserWidget_PostProcess, PostProcessTintColorAndOpacity);
+		static FName NAME_PostProcessOpacityFromTexture = GET_MEMBER_NAME_CHECKED(FVPFullScreenUserWidget_PostProcess, PostProcessOpacityFromTexture);
+
+		if (Property->GetFName() == NAME_WidgetClass
+			|| Property->GetFName() == NAME_EditorDisplayType
+			|| Property->GetFName() == NAME_PostProcessMaterial
+			|| Property->GetFName() == NAME_WidgetDrawSize
+			|| Property->GetFName() == NAME_WindowFocusable
+			|| Property->GetFName() == NAME_WindowVisibility
+			|| Property->GetFName() == NAME_ReceiveHardwareInput
+			|| Property->GetFName() == NAME_RenderTargetBackgroundColor
+			|| Property->GetFName() == NAME_RenderTargetBlendMode
+			|| Property->GetFName() == NAME_PostProcessTintColorAndOpacity
+			|| Property->GetFName() == NAME_PostProcessOpacityFromTexture)
+		{
+			bool bWasRequestedDisplay = bDisplayRequested;
+			UWorld* CurrentWorld = World.Get();
+			Hide();
+			if (bWasRequestedDisplay && CurrentWorld)
+			{
+				Display(CurrentWorld);
+			}
+		}
+	}
+
+	Super::PostEditChangeProperty(PropertyChangedEvent);
+}
+#endif
diff --git a/Engine/Plugins/Experimental/VirtualProductionUtilities/Source/VPUtilities/Private/VPFullScreenUserWidgetActor.cpp b/Engine/Plugins/Experimental/VirtualProductionUtilities/Source/VPUtilities/Private/VPFullScreenUserWidgetActor.cpp
new file mode 100644
--- /dev/null
+++ b/Engine/Plugins/Experimental/VirtualProductionUtilities/Source/VPUtilities/Private/VPFullScreenUserWidgetActor.cpp
@@ -0,0 +1,120 @@
+// Copyright 1998-2019 Epic Games, Inc. All Rights Reserved.
+
+#include "VPFullScreenUserWidgetActor.h"
+#include "VPFullScreenUserWidget.h"
+
+
+#define LOCTEXT_NAMESPACE "VPFullScreenUserWidgetActor"
+
+/////////////////////////////////////////////////////
+// AFullScreenUserWidgetActor
+
+AFullScreenUserWidgetActor::AFullScreenUserWidgetActor(const FObjectInitializer& ObjectInitializer)
+	: Super(ObjectInitializer)
+#if WITH_EDITOR
+	, bEditorDisplayRequested(false)
+#endif //WITH_EDITOR
+{
+	ScreenUserWidget = CreateDefaultSubobject<UVPFullScreenUserWidget>(TEXT("ScreenUserWidget"));
+
+	PrimaryActorTick.bCanEverTick = true;
+	PrimaryActorTick.bStartWithTickEnabled = true;
+	bAllowTickBeforeBeginPlay = true;
+	SetActorTickEnabled(true);
+	SetHidden(false);
+}
+
+void AFullScreenUserWidgetActor::PostInitializeComponents()
+{
+	Super::PostInitializeComponents();
+
+#if WITH_EDITOR
+	bEditorDisplayRequested = true;
+#endif //WITH_EDITOR
+}
+
+void AFullScreenUserWidgetActor::PostLoad()
+{
+	Super::PostLoad();
+
+#if WITH_EDITOR
+	bEditorDisplayRequested = true;
+#endif //WITH_EDITOR
+}
+
+void AFullScreenUserWidgetActor::PostActorCreated()
+{
+	Super::PostActorCreated();
+
+#if WITH_EDITOR
+	bEditorDisplayRequested = true;
+#endif //WITH_EDITOR
+}
+
+void AFullScreenUserWidgetActor::Destroyed()
+{
+	if (ScreenUserWidget)
+	{
+		ScreenUserWidget->Hide();
+	}
+	Super::Destroyed();
+}
+
+void AFullScreenUserWidgetActor::BeginPlay()
+{
+	RequestGameDisplay();
+
+	Super::BeginPlay();
+}
+
+void AFullScreenUserWidgetActor::EndPlay(const EEndPlayReason::Type EndPlayReason)
+{
+	Super::EndPlay(EndPlayReason);
+
+	if (ScreenUserWidget)
+	{
+		UWorld* ActorWorld = GetWorld();
+		if (ActorWorld && (ActorWorld->WorldType == EWorldType::Game || ActorWorld->WorldType == EWorldType::PIE))
+		{
+			ScreenUserWidget->Hide();
+		}
+	}
+}
+
+void AFullScreenUserWidgetActor::Tick(float DeltaSeconds)
+{
+	Super::Tick(DeltaSeconds);
+
+#if WITH_EDITOR
+	if (bEditorDisplayRequested)
+	{
+		bEditorDisplayRequested = false;
+		RequestEditorDisplay();
+	}
+#endif //WITH_EDITOR
+
+	if (ScreenUserWidget)
+	{
+		ScreenUserWidget->Tick(DeltaSeconds);
+	}
+}
+
+void AFullScreenUserWidgetActor::RequestEditorDisplay()
+{
+#if WITH_EDITOR
+	UWorld* ActorWorld = GetWorld();
+	if (ScreenUserWidget && ActorWorld && ActorWorld->WorldType == EWorldType::Editor)
+	{
+		ScreenUserWidget->Display(ActorWorld);
+	}
+#endif //WITH_EDITOR
+}
+
+void AFullScreenUserWidgetActor::RequestGameDisplay()
+{
+	UWorld* ActorWorld = GetWorld();
+	if (ScreenUserWidget && ActorWorld && (ActorWorld->WorldType == EWorldType::Game || ActorWorld->WorldType == EWorldType::PIE))
+	{
+		ScreenUserWidget->Display(ActorWorld);
+	}
+}
diff --git a/Engine/Plugins/Experimental/VirtualProductionUtilities/Source/VPUtilities/Public/VPFullScreenUserWidget.h b/Engine/Plugins/Experimental/VirtualProductionUtilities/Source/VPUtilities/Public/VPFullScreenUserWidget.h
new file mode 100644
--- /dev/null
+++ b/Engine/Plugins/Experimental/VirtualProductionUtilities/Source/VPUtilities/Public/VPFullScreenUserWidget.h
@@ -0,0 +1,236 @@
+// Copyright 1998-2019 Epic Games, Inc. All Rights Reserved.
+
+#pragma once
+
+#include "CoreMinimal.h"
+#include "Blueprint/UserWidget.h"
+#include "Components/WidgetComponent.h"
+#include "VPFullScreenUserWidget.generated.h"
+
+class FWidgetRenderer;
+class FVPWidgetPostProcessHitTester;
+class SConstraintCanvas;
+class SVirtualWindow;
+class SViewport;
+class SWidget;
+class ULevel;
+class UMaterialInstanceDynamic;
+class UMaterialInterface;
+class UPostProcessComponent;
+class UTextureRenderTarget2D;
+class UWorld;
+
+#if WITH_EDITOR
+class SLevelViewport;
+#endif
+
+
+UENUM(BlueprintType)
+enum class EVPWidgetDisplayType : uint8
+{
+	/** Do not display. */
+	Inactive,
+	/** Display on a game viewport. */
+	Viewport,
+	/** Display as a post process. */
+	PostProcess,
+};
+
+
+USTRUCT()
+struct FVPFullScreenUserWidget_Viewport
+{
+	GENERATED_BODY()
+
+public:
+	FVPFullScreenUserWidget_Viewport();
+	bool Display(UWorld* World, UUserWidget* Widget);
+	void Hide(UWorld* World);
+	void Tick(UWorld* World, float DeltaSeconds);
+
+private:
+	bool bAddedToGameViewport;
+
+	/** Constraint widget that contains the widget we want to display. */
+	TWeakPtr<SConstraintCanvas> FullScreenCanvasWidget;
+
+#if WITH_EDITOR
+	/** Level viewport the widget was added to. */
+	TWeakPtr<SLevelViewport> OverlayWidgetLevelViewport;
+#endif
+};
+
+USTRUCT()
+struct FVPFullScreenUserWidget_PostProcess
+{
+	GENERATED_BODY()
+
+public:
+	FVPFullScreenUserWidget_PostProcess();
+	bool Display(UWorld* World, UUserWidget* Widget);
+	void Hide(UWorld* World);
+	void Tick(UWorld* World, float DeltaSeconds);
+
+private:
+	bool CreatePostProcessComponent(UWorld* World);
+	void ReleasePostProcessComponent();
+
+	bool CreateRenderer(UWorld* World, UUserWidget* Widget);
+	void ReleaseRenderer();
+	void TickRenderer(UWorld* World, float DeltaSeconds);
+
+	FIntPoint CalculateWidgetDrawSize(UWorld* World);
+	bool IsTextureSizeValid(FIntPoint Size) const;
+
+	void RegisterHitTesterWithViewport(UWorld* World);
+	void UnRegisterHitTesterWithViewport();
+
+public:
+	/**
+	 * Post process material used to display the widget. 
+	 * SlateUI [Texture]
+	 * TintColorAndOpacity [Vector]
+	 * OpacityFromTexture [Scalar]
+	 */
+	UPROPERTY(EditAnywhere, Category = PostProcess)
+	UMaterialInterface* PostProcessMaterial;
+
+	/** Tint color and opacity for this component. */
+	UPROPERTY(EditAnywhere, Category = PostProcess)
+	FLinearColor PostProcessTintColorAndOpacity;
+
+	/** Sets the amount of opacity from the widget's UI texture to use when rendering the translucent or masked UI to the viewport (0.0-1.0). */
+	UPROPERTY(EditAnywhere, Category = PostProcess, meta = (ClampMin = 0.0f, ClampMax = 1.0f))
+	float PostProcessOpacityFromTexture;
+
+	/** The size of the rendered widget. */
+	UPROPERTY(EditAnywhere, Category = PostProcess, meta=(InlineEditConditionToggle))
+	bool bWidgetDrawSize;
+
+	/** The size of the rendered widget. */
+	UPROPERTY(EditAnywhere, Category = PostProcess, meta=(EditCondition= bWidgetDrawSize))
+	FIntPoint WidgetDrawSize;
+
+	/** Is the virtual window created to host the widget focusable? */
+	UPROPERTY(EditAnywhere, AdvancedDisplay, Category = PostProcess)
+	bool bWindowFocusable;
+
+	/** The visibility of the virtual window created to host the widget. */
+	UPROPERTY(EditAnywhere, AdvancedDisplay, Category = PostProcess)
+	EWindowVisibility WindowVisibility;
+
+	/** Register with the viewport for hardware input from the mouse and keyboard. It can and will steal focus from the viewport. */
+	UPROPERTY(EditAnywhere, AdvancedDisplay, Category= PostProcess)
+	bool bReceiveHardwareInput;
+
+	/** The background color of the render target */
+	UPROPERTY(EditAnywhere, AdvancedDisplay, Category = PostProcess)
+	FLinearColor RenderTargetBackgroundColor;
+
+	/** The blend mode for the widget. */
+	UPROPERTY(EditAnywhere, AdvancedDisplay, Category = PostProcess)
+	EWidgetBlendMode RenderTargetBlendMode;
+
+private:
+	/** Post process component used to add the material to the post process chain. */
+	UPROPERTY(Transient)
+	UPostProcessComponent* PostProcessComponent;
+
+	/** The dynamic instance of the material that the render target is attached to. */
+	UPROPERTY(Transient)
+	UMaterialInstanceDynamic* PostProcessMaterialInstance;
+
+	/** The target to which the user widget is rendered. */
+	UPROPERTY(Transient)
+	UTextureRenderTarget2D* WidgetRenderTarget;
+
+	/** The slate window that contains the user widget content. */
+	TSharedPtr<SVirtualWindow> SlateWindow;
+
+	/** The slate viewport we are registered to. */
+	TWeakPtr<SViewport> ViewportWidget;
+
+	/** Helper class for drawing widgets to a render target. */
+	FWidgetRenderer* WidgetRenderer;
+
+	/** The size of the rendered widget */
+	FIntPoint CurrentWidgetDrawSize;
+
+	/** Hit tester when we want the hardware input. */
+	TSharedPtr<FVPWidgetPostProcessHitTester> CustomHitTestPath;
+};
+
+/**
+ * Will set the Widgets on a viewport either by Widgets are first rendered to a render target, then that render target is displayed in the world.
+ */
+UCLASS(meta=(ShowOnlyInnerProperties))
+class VPUTILITIES_API UVPFullScreenUserWidget : public UObject
+{
+	GENERATED_BODY()
+
+public:
+	UVPFullScreenUserWidget(const FObjectInitializer& ObjectInitializer);
+
+	//~ Begin UObject interface
+	virtual bool IsDestructionThreadSafe() const override { return false; }
+	virtual void BeginDestroy() override;
+#if WITH_EDITOR
+	virtual void PostEditChangeProperty(FPropertyChangedEvent& PropertyChangedEvent) override;
+#endif
+	//~ End UObject Interface
+
+	bool ShouldDisplay(UWorld* World) const;
+	EVPWidgetDisplayType GetDisplayType(UWorld* World) const;
+	bool IsDisplayed() const;
+
+	virtual bool Display(UWorld* World);
+	virtual void Hide();
+	virtual void Tick(float DeltaTime);
+
+protected:
+	void InitWidget();
+	void ReleaseWidget();
+
+private:
+	void OnLevelRemovedFromWorld(ULevel* InLevel, UWorld* InWorld);
+
+protected:
+	/** The display type when the world is an editor world. */
+	UPROPERTY(EditAnywhere, Category = "User Interface")
+	EVPWidgetDisplayType EditorDisplayType;
+
+	/** The display type when the world is a game world. */
+	UPROPERTY(EditAnywhere, Category = "User Interface")
+	EVPWidgetDisplayType GameDisplayType;
+
+	/** The display type when the world is a PIE world. */
+	UPROPERTY(EditAnywhere, Category = "User Interface", meta = (DisplayName = "PIE Display Type"))
+	EVPWidgetDisplayType PIEDisplayType;
+
+	/** The class of User Widget to create and display an instance of */
+	UPROPERTY(EditAnywhere, Category = "User Interface")
+	TSubclassOf<UUserWidget> WidgetClass;
+
+protected:
+	/** Behavior when the widget should be display by the slate attached to the viewport. */
+	UPROPERTY(EditAnywhere, Category = "Viewport", meta = (ShowOnlyInnerProperties))
+	FVPFullScreenUserWidget_Viewport ViewportDisplayType;
+
+	/** Behavior when the widget should be display by a post process. */
+	UPROPERTY(EditAnywhere, Category = "Post Process", meta = (ShowOnlyInnerProperties))
+	FVPFullScreenUserWidget_PostProcess PostProcessDisplayType;
+
+private:
+	/** The User Widget object displayed and managed by this component */
+	UPROPERTY(Transient, DuplicateTransient)
+	UUserWidget* Widget;
+
+	/** The world the widget is attached to. */
+	TWeakObjectPtr<UWorld> World;
+
+	/** How we currently displaying the widget. */
+	EVPWidgetDisplayType CurrentDisplayType;
+
+	/** The user requested the widget to be displayed. It's possible that some setting are invalid and the widget will not be displayed. */
+	bool bDisplayRequested;
+};
diff --git a/Engine/Plugins/Experimental/VirtualProductionUtilities/Source/VPUtilities/Public/VPFullScreenUserWidgetActor.h b/Engine/Plugins/Experimental/VirtualProductionUtilities/Source/VPUtilities/Public/VPFullScreenUserWidgetActor.h
new file mode 100644
--- /dev/null
+++ b/Engine/Plugins/Experimental/VirtualProductionUtilities/Source/VPUtilities/Public/VPFullScreenUserWidgetActor.h
@@ -0,0 +1,51 @@
+// Copyright 1998-2019 Epic Games, Inc. All Rights Reserved.
+
+#pragma once
+
+#include "CoreMinimal.h"
+#include "GameFramework/Info.h"
+#include "VPFullScreenUserWidgetActor.generated.h"
+
+
+class UVPFullScreenUserWidget;
+
+/**
+ * Widgets are first rendered to a render target, then that render target is displayed in the world.
+ */
+UCLASS(Blueprintable, ClassGroup = "UserInterface", HideCategories = (Actor, Input, Movement, Collision, Rendering, "Utilities|Transformation", LOD), ShowCategories = ("Input|MouseInput", "Input|TouchInput"))
+class VPUTILITIES_API AFullScreenUserWidgetActor : public AInfo
+{
+	GENERATED_BODY()
+
+public:
+	AFullScreenUserWidgetActor(const FObjectInitializer& ObjectInitializer);
+
+	//~ Begin AActor interface
+	virtual void PostInitializeComponents() override;
+	virtual void PostLoad() override;
+	virtual void PostActorCreated() override;
+	virtual void Destroyed() override;
+
+	virtual void BeginPlay() override;
+	virtual void EndPlay(const EEndPlayReason::Type EndPlayReason) override;
+	virtual void Tick(float DeltaSeconds) override;
+
+#if WITH_EDITOR
+		virtual bool ShouldTickIfViewportsOnly() const override { return true; }
+#endif //WITH_EDITOR
+	//~ End AActor Interface
+
+private:
+	void RequestEditorDisplay();
+	void RequestGameDisplay();
+
+protected:
+	/** */
+	UPROPERTY(VisibleAnywhere, Instanced, NoClear, Category = "User Interface", meta = (ShowOnlyInnerProperties))
+	UVPFullScreenUserWidget* ScreenUserWidget;
+
+#if WITH_EDITORONLY_DATA
+	/** Display requested and will be executed on the first frame because we can't call BP function in the loading phase */
+	bool bEditorDisplayRequested;
+#endif
+};
diff --git a/Engine/Plugins/Experimental/VirtualProductionUtilities/Source/VPUtilities/VPUtilities.Build.cs b/Engine/Plugins/Experimental/VirtualProductionUtilities/Source/VPUtilities/VPUtilities.Build.cs
--- a/Engine/Plugins/Experimental/VirtualProductionUtilities/Source/VPUtilities/VPUtilities.Build.cs
+++ b/Engine/Plugins/Experimental/VirtualProductionUtilities/Source/VPUtilities/VPUtilities.Build.cs
@@ -17,22 +17,34 @@ public class VPUtilities : ModuleRules
 				"Engine",
 				"GameplayTags",
 				"TimeManagement",
+				"UMG",
 				"VPBookmark",
 			}
 		);
 
+		PrivateDependencyModuleNames.AddRange(
+			new string[]
+			{
+				"RenderCore",
+				"RHI",
+				"Slate",
+				"SlateCore",
+			}
+		);
+
 		if (Target.Type == TargetType.Editor)
 		{
 			PrivateDependencyModuleNames.AddRange(
 				new string[]
 				{
-					"VPBookmarkEditor",
+					"Blutility",
+					"LevelEditor",
 					"UnrealEd",
 					"ViewportInteraction",
+					"VPBookmarkEditor",
 					"VREditor",
-					"Blutility",
-                }
-            );
+				}
+			);
 		}
 	}
 }
