Fix deprecation warning about UProperty

#rnx

diff --git a/Engine/Plugins/Experimental/VirtualProductionUtilities/Source/VPUtilities/Private/VPFullScreenUserWidget.cpp b/Engine/Plugins/Experimental/VirtualProductionUtilities/Source/VPUtilities/Private/VPFullScreenUserWidget.cpp
--- a/Engine/Plugins/Experimental/VirtualProductionUtilities/Source/VPUtilities/Private/VPFullScreenUserWidget.cpp
+++ b/Engine/Plugins/Experimental/VirtualProductionUtilities/Source/VPUtilities/Private/VPFullScreenUserWidget.cpp
@@ -674,7 +674,7 @@ void UVPFullScreenUserWidget::OnLevelRemovedFromWorld(ULevel* InLevel, UWorld* I
 #if WITH_EDITOR
 void UVPFullScreenUserWidget::PostEditChangeProperty(FPropertyChangedEvent& PropertyChangedEvent)
 {
-	UProperty* Property = PropertyChangedEvent.MemberProperty;
+	FProperty* Property = PropertyChangedEvent.MemberProperty;
 
 	if (Property && PropertyChangedEvent.ChangeType != EPropertyChangeType::Interactive)
 	{
