Fix CIS missing include.
#rnx

diff --git a/Engine/Plugins/Animation/LiveLink/Source/LiveLink/Private/LiveLinkTimedDataInput.cpp b/Engine/Plugins/Animation/LiveLink/Source/LiveLink/Private/LiveLinkTimedDataInput.cpp
--- a/Engine/Plugins/Animation/LiveLink/Source/LiveLink/Private/LiveLinkTimedDataInput.cpp
+++ b/Engine/Plugins/Animation/LiveLink/Source/LiveLink/Private/LiveLinkTimedDataInput.cpp
@@ -2,6 +2,7 @@
 
 #include "LiveLinkTimedDataInput.h"
 
+#include "HAL/IConsoleManager.h"
 #include "ILiveLinkModule.h"
 #include "ITimeManagementModule.h"
 #include "LiveLinkClient.h"
diff --git a/Engine/Plugins/Experimental/VirtualProductionUtilities/Source/VPUtilities/Private/VPFullScreenUserWidget.cpp b/Engine/Plugins/Experimental/VirtualProductionUtilities/Source/VPUtilities/Private/VPFullScreenUserWidget.cpp
--- a/Engine/Plugins/Experimental/VirtualProductionUtilities/Source/VPUtilities/Private/VPFullScreenUserWidget.cpp
+++ b/Engine/Plugins/Experimental/VirtualProductionUtilities/Source/VPUtilities/Private/VPFullScreenUserWidget.cpp
@@ -3,6 +3,7 @@
 #include "VPFullScreenUserWidget.h"
 
 #include "Components/PostProcessComponent.h"
+#include "Engine/Engine.h"
 #include "Engine/Texture2D.h"
 #include "Engine/TextureRenderTarget2D.h"
 #include "Engine/UserInterfaceSettings.h"
