Teleport to bookmark takes headset offset into account.

diff --git a/Engine/Plugins/Experimental/VirtualProductionUtilities/Source/VPBookmarkEditor/Private/VPBookmarkTypeActions.cpp b/Engine/Plugins/Experimental/VirtualProductionUtilities/Source/VPBookmarkEditor/Private/VPBookmarkTypeActions.cpp
--- a/Engine/Plugins/Experimental/VirtualProductionUtilities/Source/VPBookmarkEditor/Private/VPBookmarkTypeActions.cpp
+++ b/Engine/Plugins/Experimental/VirtualProductionUtilities/Source/VPBookmarkEditor/Private/VPBookmarkTypeActions.cpp
@@ -58,8 +58,8 @@ void FVPBookmarkTypeActions::ActivateBookmark(UVPBookmark* InBookmark, FEditorVi
 		const FVector ActorLocation = ActorTransform.GetLocation();
 
 		// Set Location
-		FVector Offset = ActorRotation.RotateVector(InBookmark->CachedViewportData.JumpToOffsetLocation);
-		Transform.SetLocation(ActorLocation - Offset);
+		FVector Offset = FVector(0, 0, 0);
+		
 
 		// Set Rotation
 		bool bIsRotationSet = false;
@@ -83,6 +83,8 @@ void FVPBookmarkTypeActions::ActivateBookmark(UVPBookmark* InBookmark, FEditorVi
 			}
 		}
 
+		Transform.SetLocation(ActorLocation - Offset);
+
 		if (!bIsRotationSet)
 		{
 			if (InBookmark->CachedViewportData.bFlattenRotation)
