Automatically turn off the snapping functionality when going into VR with the VR scouting plugin enabled, and then turn snapping back on when exiting VR. This avoids issue of not being able to grab items.

diff --git a/Engine/Plugins/Experimental/VirtualProductionUtilities/Source/VPUtilitiesEditor/Private/VPScoutingSubsystem.cpp b/Engine/Plugins/Experimental/VirtualProductionUtilities/Source/VPUtilitiesEditor/Private/VPScoutingSubsystem.cpp
--- a/Engine/Plugins/Experimental/VirtualProductionUtilities/Source/VPUtilitiesEditor/Private/VPScoutingSubsystem.cpp
+++ b/Engine/Plugins/Experimental/VirtualProductionUtilities/Source/VPUtilitiesEditor/Private/VPScoutingSubsystem.cpp
@@ -16,6 +16,7 @@
 #include "UObject/Script.h"
 #include "VPSettings.h"
 #include "VPUtilitiesEditorSettings.h"
+#include "LevelEditorActions.h"
 
 /* UVPScoutingSubsystem name
  *****************************************************************************/
@@ -379,3 +380,13 @@ void UVPScoutingSubsystem::ExitVRMode()
 		VREditorModule.EnableVREditor(false);
 	}
 }
+
+bool UVPScoutingSubsystem::IsGridSnappingEnabled()
+{
+	return FLevelEditorActionCallbacks::LocationGridSnap_IsChecked();
+}
+
+void UVPScoutingSubsystem::ToggleGridSnapping()
+{
+	FLevelEditorActionCallbacks::LocationGridSnap_Clicked();
+}
diff --git a/Engine/Plugins/Experimental/VirtualProductionUtilities/Source/VPUtilitiesEditor/Public/VPScoutingSubsystem.h b/Engine/Plugins/Experimental/VirtualProductionUtilities/Source/VPUtilitiesEditor/Public/VPScoutingSubsystem.h
--- a/Engine/Plugins/Experimental/VirtualProductionUtilities/Source/VPUtilitiesEditor/Public/VPScoutingSubsystem.h
+++ b/Engine/Plugins/Experimental/VirtualProductionUtilities/Source/VPUtilitiesEditor/Public/VPScoutingSubsystem.h
@@ -214,6 +214,14 @@ public:
 	UFUNCTION(BlueprintCallable, Category = "Virtual Production")
 	static void ExitVRMode();
 
+	/** Whether editor grid snapping is enabled */
+	UFUNCTION(BlueprintCallable, Category = "Virtual Production")
+	static bool IsGridSnappingEnabled();
+
+	/** Toggle editor grid snapping */
+	UFUNCTION(BlueprintCallable, Category = "Virtual Production")
+	static void ToggleGridSnapping();
+
 private:
 
 	FDelegateHandle EngineInitCompleteDelegate;
