Added CCV content to VirtualProductionUtilities plugin.
This required added dependency on Composure.

diff --git a/Engine/Plugins/Experimental/VirtualProductionUtilities/VirtualProductionUtilities.uplugin b/Engine/Plugins/Experimental/VirtualProductionUtilities/VirtualProductionUtilities.uplugin
--- a/Engine/Plugins/Experimental/VirtualProductionUtilities/VirtualProductionUtilities.uplugin
+++ b/Engine/Plugins/Experimental/VirtualProductionUtilities/VirtualProductionUtilities.uplugin
@@ -51,6 +51,10 @@
 		{
 			"Name": "Takes",
 			"Enabled": true
+		},
+		{
+			"Name": "Composure",
+			"Enabled": true
 		}
 	]
 }
