[NO NON-PII COMMIT MESSAGE]

diff --git a/Engine/Plugins/Experimental/VirtualProductionUtilities/Source/VPUtilitiesEditor/Private/VPScoutingSubsystem.cpp b/Engine/Plugins/Experimental/VirtualProductionUtilities/Source/VPUtilitiesEditor/Private/VPScoutingSubsystem.cpp
--- a/Engine/Plugins/Experimental/VirtualProductionUtilities/Source/VPUtilitiesEditor/Private/VPScoutingSubsystem.cpp
+++ b/Engine/Plugins/Experimental/VirtualProductionUtilities/Source/VPUtilitiesEditor/Private/VPScoutingSubsystem.cpp
@@ -381,12 +381,22 @@ void UVPScoutingSubsystem::ExitVRMode()
 	}
 }
 
-bool UVPScoutingSubsystem::IsGridSnappingEnabled()
+bool UVPScoutingSubsystem::IsLocationGridSnappingEnabled()
 {
 	return FLevelEditorActionCallbacks::LocationGridSnap_IsChecked();
 }
 
-void UVPScoutingSubsystem::ToggleGridSnapping()
+void UVPScoutingSubsystem::ToggleLocationGridSnapping()
 {
 	FLevelEditorActionCallbacks::LocationGridSnap_Clicked();
 }
+
+bool UVPScoutingSubsystem::IsRotationGridSnappingEnabled()
+{
+	return FLevelEditorActionCallbacks::RotationGridSnap_IsChecked();
+}
+
+void UVPScoutingSubsystem::ToggleRotationGridSnapping()
+{
+	FLevelEditorActionCallbacks::RotationGridSnap_Clicked();
+}
diff --git a/Engine/Plugins/Experimental/VirtualProductionUtilities/Source/VPUtilitiesEditor/Public/VPScoutingSubsystem.h b/Engine/Plugins/Experimental/VirtualProductionUtilities/Source/VPUtilitiesEditor/Public/VPScoutingSubsystem.h
--- a/Engine/Plugins/Experimental/VirtualProductionUtilities/Source/VPUtilitiesEditor/Public/VPScoutingSubsystem.h
+++ b/Engine/Plugins/Experimental/VirtualProductionUtilities/Source/VPUtilitiesEditor/Public/VPScoutingSubsystem.h
@@ -214,13 +214,21 @@ public:
 	UFUNCTION(BlueprintCallable, Category = "Virtual Production")
 	static void ExitVRMode();
 
-	/** Whether editor grid snapping is enabled */
+	/** Whether location grid snapping is enabled */
 	UFUNCTION(BlueprintCallable, Category = "Virtual Production")
-	static bool IsGridSnappingEnabled();
+	static bool IsLocationGridSnappingEnabled();
 
-	/** Toggle editor grid snapping */
+	/** Toggle location grid snapping */
 	UFUNCTION(BlueprintCallable, Category = "Virtual Production")
-	static void ToggleGridSnapping();
+	static void ToggleLocationGridSnapping();
+
+	/** Whether rotation grid snapping is enabled */
+	UFUNCTION(BlueprintCallable, Category = "Virtual Production")
+		static bool IsRotationGridSnappingEnabled();
+
+	/** Toggle rotation grid snapping */
+	UFUNCTION(BlueprintCallable, Category = "Virtual Production")
+		static void ToggleRotationGridSnapping();
 
 private:
 
