[NO NON-PII COMMIT MESSAGE]

diff --git a/Engine/Plugins/Experimental/VirtualProductionUtilities/Source/VPUtilities/Private/VPViewportTickableActorBase.cpp b/Engine/Plugins/Experimental/VirtualProductionUtilities/Source/VPUtilities/Private/VPViewportTickableActorBase.cpp
--- a/Engine/Plugins/Experimental/VirtualProductionUtilities/Source/VPUtilities/Private/VPViewportTickableActorBase.cpp
+++ b/Engine/Plugins/Experimental/VirtualProductionUtilities/Source/VPUtilities/Private/VPViewportTickableActorBase.cpp
@@ -66,13 +66,16 @@ void AVPViewportTickableActorBase::Destroyed()
 	Super::Destroyed();
 }
 
+void AVPViewportTickableActorBase::LockLocation(bool bSetLockLocation)
+{
+	bLockLocation = bSetLockLocation;
+}
 
 void AVPViewportTickableActorBase::EditorTick_Implementation(float DeltaSeconds)
 {
 
 }
 
-
 void AVPViewportTickableActorBase::EditorDestroyed_Implementation()
 {
 
diff --git a/Engine/Plugins/Experimental/VirtualProductionUtilities/Source/VPUtilities/Public/VPViewportTickableActorBase.h b/Engine/Plugins/Experimental/VirtualProductionUtilities/Source/VPUtilities/Public/VPViewportTickableActorBase.h
--- a/Engine/Plugins/Experimental/VirtualProductionUtilities/Source/VPUtilities/Public/VPViewportTickableActorBase.h
+++ b/Engine/Plugins/Experimental/VirtualProductionUtilities/Source/VPUtilities/Public/VPViewportTickableActorBase.h
@@ -35,6 +35,10 @@ public:
 	UFUNCTION(BlueprintNativeEvent, CallInEditor, BlueprintCallable, Category = "Utilities")
 	void EditorDestroyed();
 
+	/** Sets the LockLocation variable to disable movement from the translation gizmo */
+	UFUNCTION(CallInEditor, BlueprintCallable, Category = "Editor")
+	void LockLocation(bool bSetLockLocation);
+
 	/**
 	 * Where the actor should be ticked.
 	 * Editor = Tick in the editor viewport. Use the event EditorTick.
diff --git a/Engine/Plugins/Experimental/VirtualProductionUtilities/Source/VPUtilitiesEditor/Private/VPEditorTickableActorBase.cpp b/Engine/Plugins/Experimental/VirtualProductionUtilities/Source/VPUtilitiesEditor/Private/VPEditorTickableActorBase.cpp
--- a/Engine/Plugins/Experimental/VirtualProductionUtilities/Source/VPUtilitiesEditor/Private/VPEditorTickableActorBase.cpp
+++ b/Engine/Plugins/Experimental/VirtualProductionUtilities/Source/VPUtilitiesEditor/Private/VPEditorTickableActorBase.cpp
@@ -3,8 +3,4 @@
 #include "VPEditorTickableActorBase.h"
 
 
- void AVPEditorTickableActorBase::LockLocation(bool bSetLockLocation) 
- {
-	 bLockLocation = bSetLockLocation;
- }
  
\ No newline at end of file
diff --git a/Engine/Plugins/Experimental/VirtualProductionUtilities/Source/VPUtilitiesEditor/Public/VPEditorTickableActorBase.h b/Engine/Plugins/Experimental/VirtualProductionUtilities/Source/VPUtilitiesEditor/Public/VPEditorTickableActorBase.h
--- a/Engine/Plugins/Experimental/VirtualProductionUtilities/Source/VPUtilitiesEditor/Public/VPEditorTickableActorBase.h
+++ b/Engine/Plugins/Experimental/VirtualProductionUtilities/Source/VPUtilitiesEditor/Public/VPEditorTickableActorBase.h
@@ -17,7 +17,5 @@ class VPUTILITIESEDITOR_API AVPEditorTickableActorBase : public AVPViewportTicka
 
 public:
 
-	/** Sets the LockLocation variable to disable movement from the translation gizmo */
-	UFUNCTION(CallInEditor, BlueprintCallable, Category = "Editor")
-	void LockLocation(bool bSetLockLocation);
+	
 };
