Rename VPViewportTickableActorBase.LockLocation to EditorLockLocation
#rnx

diff --git a/Engine/Plugins/Experimental/VirtualProductionUtilities/Config/BaseVirtualProductionUtilities.ini b/Engine/Plugins/Experimental/VirtualProductionUtilities/Config/BaseVirtualProductionUtilities.ini
--- a/Engine/Plugins/Experimental/VirtualProductionUtilities/Config/BaseVirtualProductionUtilities.ini
+++ b/Engine/Plugins/Experimental/VirtualProductionUtilities/Config/BaseVirtualProductionUtilities.ini
@@ -1,6 +1,7 @@
 [CoreRedirects]
 
 +ClassRedirects=(OldName="VRTool",NewName="/Script/VPUtilitiesEditor.VRTool")
++FunctionRedirects=(OldName="VPViewportTickableActorBase.LockLocation",NewName="VPViewportTickableActorBase.EditorLockLocation")
 
 
 [/Script/VPUtilitiesEditor.VPUtilitiesEditorSettings]
diff --git a/Engine/Plugins/Experimental/VirtualProductionUtilities/Source/VPUtilities/Private/VPViewportTickableActorBase.cpp b/Engine/Plugins/Experimental/VirtualProductionUtilities/Source/VPUtilities/Private/VPViewportTickableActorBase.cpp
--- a/Engine/Plugins/Experimental/VirtualProductionUtilities/Source/VPUtilities/Private/VPViewportTickableActorBase.cpp
+++ b/Engine/Plugins/Experimental/VirtualProductionUtilities/Source/VPUtilities/Private/VPViewportTickableActorBase.cpp
@@ -66,9 +66,11 @@ void AVPViewportTickableActorBase::Destroyed()
 	Super::Destroyed();
 }
 
-void AVPViewportTickableActorBase::LockLocation(bool bSetLockLocation)
+void AVPViewportTickableActorBase::EditorLockLocation(bool bSetLockLocation)
 {
+#if WITH_EDITOR
 	bLockLocation = bSetLockLocation;
+#endif
 }
 
 void AVPViewportTickableActorBase::EditorTick_Implementation(float DeltaSeconds)
diff --git a/Engine/Plugins/Experimental/VirtualProductionUtilities/Source/VPUtilities/Public/VPViewportTickableActorBase.h b/Engine/Plugins/Experimental/VirtualProductionUtilities/Source/VPUtilities/Public/VPViewportTickableActorBase.h
--- a/Engine/Plugins/Experimental/VirtualProductionUtilities/Source/VPUtilities/Public/VPViewportTickableActorBase.h
+++ b/Engine/Plugins/Experimental/VirtualProductionUtilities/Source/VPUtilities/Public/VPViewportTickableActorBase.h
@@ -37,7 +37,7 @@ public:
 
 	/** Sets the LockLocation variable to disable movement from the translation gizmo */
 	UFUNCTION(CallInEditor, BlueprintCallable, Category = "Editor")
-	void LockLocation(bool bSetLockLocation);
+	void EditorLockLocation(bool bSetLockLocation);
 
 	/**
 	 * Where the actor should be ticked.
