Fixing some deprecated function usages.

diff --git a/Engine/Plugins/Experimental/VirtualProductionUtilities/Source/VPUtilities/Private/VPFullScreenUserWidget.cpp b/Engine/Plugins/Experimental/VirtualProductionUtilities/Source/VPUtilities/Private/VPFullScreenUserWidget.cpp
--- a/Engine/Plugins/Experimental/VirtualProductionUtilities/Source/VPUtilities/Private/VPFullScreenUserWidget.cpp
+++ b/Engine/Plugins/Experimental/VirtualProductionUtilities/Source/VPUtilities/Private/VPFullScreenUserWidget.cpp
@@ -93,7 +93,7 @@ public:
 		return ArrangedWidgets;
 	}
 
-	virtual void ArrangeChildren(FArrangedChildren& ArrangedChildren) const override
+	virtual void ArrangeCustomHitTestChildren(FArrangedChildren& ArrangedChildren) const override
 	{
 		// Add the displayed slate to the list of widgets.
 		if (TSharedPtr<SVirtualWindow> SlateWindowPin = SlateWindow.Pin())
@@ -103,7 +103,7 @@ public:
 		}
 	}
 
-	virtual TSharedPtr<struct FVirtualPointerPosition> TranslateMouseCoordinateFor3DChild(const TSharedRef<SWidget>& ChildWidget, const FGeometry& ViewportGeometry, const FVector2D& ScreenSpaceMouseCoordinate, const FVector2D& LastScreenSpaceMouseCoordinate) const override
+	virtual TSharedPtr<struct FVirtualPointerPosition> TranslateMouseCoordinateForCustomHitTestChild(const TSharedRef<SWidget>& ChildWidget, const FGeometry& ViewportGeometry, const FVector2D& ScreenSpaceMouseCoordinate, const FVector2D& LastScreenSpaceMouseCoordinate) const override
 	{
 		return nullptr;
 	}
