Fix missing virtual production placement browser icon.

diff --git a/Engine/Plugins/Experimental/VirtualProductionUtilities/Source/VPUtilitiesEditor/Private/VPUtilitiesEditorModule.cpp b/Engine/Plugins/Experimental/VirtualProductionUtilities/Source/VPUtilitiesEditor/Private/VPUtilitiesEditorModule.cpp
--- a/Engine/Plugins/Experimental/VirtualProductionUtilities/Source/VPUtilitiesEditor/Private/VPUtilitiesEditorModule.cpp
+++ b/Engine/Plugins/Experimental/VirtualProductionUtilities/Source/VPUtilitiesEditor/Private/VPUtilitiesEditorModule.cpp
@@ -90,6 +90,7 @@ const FPlacementCategoryInfo* FVPUtilitiesEditorModule::GetVirtualProductionPlac
 		{
 			FPlacementCategoryInfo Info(
 				LOCTEXT("VirtualProductionCategoryName", "Virtual Production"),
+				FSlateIcon(FVPUtilitiesEditorStyle::GetStyleSetName(), "PlacementBrowser.Icons.VirtualProduction"),
 				PlacementModeCategoryHandle,
 				TEXT("PMVirtualProduction"),
 				25
diff --git a/Engine/Plugins/Experimental/VirtualProductionUtilities/Source/VPUtilitiesEditor/Private/VPUtilitiesEditorStyle.cpp b/Engine/Plugins/Experimental/VirtualProductionUtilities/Source/VPUtilitiesEditor/Private/VPUtilitiesEditorStyle.cpp
--- a/Engine/Plugins/Experimental/VirtualProductionUtilities/Source/VPUtilitiesEditor/Private/VPUtilitiesEditorStyle.cpp
+++ b/Engine/Plugins/Experimental/VirtualProductionUtilities/Source/VPUtilitiesEditor/Private/VPUtilitiesEditorStyle.cpp
@@ -6,6 +6,10 @@
 #include "Styling/SlateStyle.h"
 #include "Styling/SlateStyleRegistry.h"
 
+#define RootToContentDir VPUtilitiesEditorStyle::StyleInstance->RootToContentDir
+
+#include "Styling/SlateStyleMacros.h"
+
 namespace VPUtilitiesEditorStyle
 {
 	const FName NAME_StyleName(TEXT("VPUtilitiesStyle"));
@@ -13,8 +17,6 @@ namespace VPUtilitiesEditorStyle
 	static TUniquePtr<FSlateStyleSet> StyleInstance;
 }
 
-#define IMAGE_BRUSH(RelativePath, ...) FSlateImageBrush(VPUtilitiesEditorStyle::StyleInstance->RootToContentDir(RelativePath, TEXT(".png")), __VA_ARGS__)
-
 void FVPUtilitiesEditorStyle::Register()
 {
 	const FVector2D Icon16x16(16.0f, 16.0f);
@@ -24,6 +26,8 @@ void FVPUtilitiesEditorStyle::Register()
 
 	VPUtilitiesEditorStyle::StyleInstance->Set("TabIcons.Genlock.Small", new IMAGE_BRUSH("Icon_GenlockTab_16x", Icon16x16));
 
+	VPUtilitiesEditorStyle::StyleInstance->Set("PlacementBrowser.Icons.VirtualProduction", new IMAGE_BRUSH_SVG("VirtualProduction", Icon16x16));
+
 
 	FSlateStyleRegistry::RegisterSlateStyle(*VPUtilitiesEditorStyle::StyleInstance.Get());
 }
@@ -45,4 +49,4 @@ const ISlateStyle& FVPUtilitiesEditorStyle::Get()
 	return *VPUtilitiesEditorStyle::StyleInstance.Get();
 }
 
-#undef IMAGE_BRUSH
+#undef RootToContentDir
