Fix CIS issue by create wrapper method for GetEditorWorld in VPUtilities blueprint library.

diff --git a/Engine/Plugins/Experimental/VirtualProductionUtilities/Source/VPUtilities/Private/VPBlueprintLibrary.cpp b/Engine/Plugins/Experimental/VirtualProductionUtilities/Source/VPUtilities/Private/VPBlueprintLibrary.cpp
--- a/Engine/Plugins/Experimental/VirtualProductionUtilities/Source/VPUtilities/Private/VPBlueprintLibrary.cpp
+++ b/Engine/Plugins/Experimental/VirtualProductionUtilities/Source/VPUtilities/Private/VPBlueprintLibrary.cpp
@@ -12,6 +12,7 @@
 #include "Editor.h"
 #include "IVREditorModule.h"
 #include "LevelEditorViewport.h"
+#include "Subsystems/UnrealEditorSubsystem.h"
 #include "ViewportWorldInteraction.h"
 #include "VPBookmarkEditorBlueprintLibrary.h"
 #include "VREditorInteractor.h"
@@ -268,6 +269,19 @@ bool UVPBlueprintLibrary::EditorDuplicate()
 	return false;
 }
 
+UWorld* UVPBlueprintLibrary::GetEditorWorld()
+{
+	UWorld* World = nullptr;
+#if WITH_EDITOR
+	if (UUnrealEditorSubsystem* UnrealEditorSubsystem = GEditor->GetEditorSubsystem<UUnrealEditorSubsystem>())
+	{
+		World = UnrealEditorSubsystem->GetEditorWorld();
+	}
+
+#endif WITH_EDITOR
+	return World;
+}
+
 void UVPBlueprintLibrary::VPBookmarkSplineMeshIndicatorSetStartAndEnd(USplineMeshComponent* SplineMesh)
 {
 	SplineMesh->SetVisibility(true);
@@ -282,4 +296,3 @@ void UVPBlueprintLibrary::VPBookmarkSplineMeshIndicatorDisable(USplineMeshCompon
 {
 	SplineMesh->SetVisibility(false);
 }
-
diff --git a/Engine/Plugins/Experimental/VirtualProductionUtilities/Source/VPUtilities/Public/VPBlueprintLibrary.h b/Engine/Plugins/Experimental/VirtualProductionUtilities/Source/VPUtilities/Public/VPBlueprintLibrary.h
--- a/Engine/Plugins/Experimental/VirtualProductionUtilities/Source/VPUtilities/Public/VPBlueprintLibrary.h
+++ b/Engine/Plugins/Experimental/VirtualProductionUtilities/Source/VPUtilities/Public/VPBlueprintLibrary.h
@@ -82,9 +82,18 @@ public:
 	UFUNCTION(BlueprintCallable, Category = "Virtual Production")
 	static bool EditorDeleteSelectedObjects();
 
+	/**
+	 * Wrapper around UUnrealEditorSubsystem::GetEditorWorld.
+	 * Used because you can't get the subsystem without being an editor utility actor.
+	 */
+	UFUNCTION(BlueprintCallable, Category = "Virtual Production")
+	static UWorld* GetEditorWorld();
+
 	UFUNCTION(BlueprintCallable, Category = "VPBookmarks")		
 	static void VPBookmarkSplineMeshIndicatorSetStartAndEnd(USplineMeshComponent* SplineMesh);
 
 	UFUNCTION(BlueprintCallable, Category = "VPBookmarks")
 	static void VPBookmarkSplineMeshIndicatorDisable(USplineMeshComponent* SplineMesh);
+
+	
 };
\ No newline at end of file
