Fix faulty preprocessor directive in VPBlueprintLibrary

diff --git a/Engine/Plugins/Experimental/VirtualProductionUtilities/Source/VPUtilities/Private/VPBlueprintLibrary.cpp b/Engine/Plugins/Experimental/VirtualProductionUtilities/Source/VPUtilities/Private/VPBlueprintLibrary.cpp
--- a/Engine/Plugins/Experimental/VirtualProductionUtilities/Source/VPUtilities/Private/VPBlueprintLibrary.cpp
+++ b/Engine/Plugins/Experimental/VirtualProductionUtilities/Source/VPUtilities/Private/VPBlueprintLibrary.cpp
@@ -278,7 +278,7 @@ UWorld* UVPBlueprintLibrary::GetEditorWorld()
 		World = UnrealEditorSubsystem->GetEditorWorld();
 	}
 
-#endif WITH_EDITOR
+#endif
 	return World;
 }
 
