Moving TimecodeProviderTab out of TimeManagement engine module to VirtualProductionUtilities plugin.

diff --git a/Engine/Plugins/Experimental/VirtualProductionUtilities/Source/VPUtilitiesEditor/Private/STimecodeProvider.cpp b/Engine/Plugins/Experimental/VirtualProductionUtilities/Source/VPUtilitiesEditor/Private/STimecodeProvider.cpp
new file mode 100644
--- /dev/null
+++ b/Engine/Plugins/Experimental/VirtualProductionUtilities/Source/VPUtilitiesEditor/Private/STimecodeProvider.cpp
@@ -0,0 +1,148 @@
+// Copyright Epic Games, Inc. All Rights Reserved.
+
+#include "STimecodeProvider.h"
+
+#include "EditorFontGlyphs.h"
+#include "EditorStyleSet.h"
+#include "Engine/Engine.h"
+#include "Misc/App.h"
+#include "Widgets/SBoxPanel.h"
+#include "Widgets/SNullWidget.h"
+#include "Widgets/STimecode.h"
+#include "Widgets/Images/SImage.h"
+#include "Widgets/Text/STextBlock.h"
+
+void STimecodeProvider::Construct(const FArguments& InArgs)
+{
+	OverrideTimecodeProvider = InArgs._OverrideTimecodeProvider;
+
+	TSharedRef<SWidget> StateWidget = InArgs._DisplaySynchronizationState ? SNew(STextBlock)
+		.Font(FEditorStyle::Get().GetFontStyle("FontAwesome.11"))
+		.Text(this, &STimecodeProvider::HandleStateText)
+		.ColorAndOpacity(this, &STimecodeProvider::HandleIconColorAndOpacity)
+		: SNullWidget::NullWidget;
+
+	TSharedRef<SWidget> FrameRateWidget = InArgs._DisplayFrameRate ? SNew(STextBlock)
+		.Text(MakeAttributeLambda([this]
+		{
+			if (const UTimecodeProvider* TimecodeProviderPtr = GetTimecodeProvider())
+			{
+				return TimecodeProviderPtr->GetFrameRate().ToPrettyText();
+			}
+			return NSLOCTEXT("TimecodeProvider", "Undefined", "<Undefined>");
+		}))
+		.Font(InArgs._TimecodeProviderFont)
+		.ColorAndOpacity(InArgs._TimecodeProviderColor)
+		: SNullWidget::NullWidget;
+
+	ChildSlot
+	[
+		SNew(SVerticalBox)
+		+ SVerticalBox::Slot()
+		.Padding(0)
+		.AutoHeight()
+		[
+			SNew(SHorizontalBox)
+			+ SHorizontalBox::Slot()
+			.Padding(0.f, 0.f, 4.f, 0.f)
+			.VAlign(VAlign_Center)
+			.AutoWidth()
+			[
+				StateWidget
+			]
+			+ SHorizontalBox::Slot()
+			.VAlign(VAlign_Center)
+			.AutoWidth()
+			[
+				SNew(STextBlock)
+				.Text(MakeAttributeLambda([this]
+				{
+					if (const UTimecodeProvider* TimecodeProviderPtr = GetTimecodeProvider())
+					{
+						return FText::FromName(TimecodeProviderPtr->GetFName());
+					}
+					return NSLOCTEXT("TimecodeProvider", "Undefined", "<Undefined>");
+				}))
+				.Font(InArgs._TimecodeProviderFont)
+				.ColorAndOpacity(InArgs._TimecodeProviderColor)
+			]
+			+ SHorizontalBox::Slot()
+			.VAlign(VAlign_Center)
+			.HAlign(HAlign_Right)
+			.FillWidth(1.f)
+			[
+				FrameRateWidget
+			]
+		]
+		+ SVerticalBox::Slot()
+		.Padding(0, -4, 0, 0)
+		.AutoHeight()
+		[
+			SNew(STimecode)
+			.Timecode(MakeAttributeLambda([this]
+				{
+					if (const UTimecodeProvider* OverrideTimecodeProviderPtr = OverrideTimecodeProvider.Get().Get())
+					{
+						return OverrideTimecodeProviderPtr->GetDelayedTimecode();
+					}
+
+					// If we use the Engine's TimecodeProvider, get the timecode for the current frame
+					return FApp::GetTimecode();
+				}))
+			.DisplayLabel(InArgs._DisplayLabel)
+			.TimecodeFont(InArgs._TimecodeFont)
+			.TimecodeColor(InArgs._TimecodeColor)
+		]
+	];
+}
+
+const UTimecodeProvider* STimecodeProvider::GetTimecodeProvider() const
+{
+	const UTimecodeProvider* TimecodeProviderPtr = OverrideTimecodeProvider.Get().Get();
+	if (!TimecodeProviderPtr)
+	{
+		TimecodeProviderPtr = GEngine->GetTimecodeProvider();
+	}
+
+	return TimecodeProviderPtr;
+}
+
+FSlateColor STimecodeProvider::HandleIconColorAndOpacity() const
+{
+	FSlateColor Result = FSlateColor::UseForeground();
+	if (const UTimecodeProvider* TimecodeProviderPtr = GetTimecodeProvider())
+	{
+		ETimecodeProviderSynchronizationState State = TimecodeProviderPtr->GetSynchronizationState();
+		switch (State)
+		{
+		case ETimecodeProviderSynchronizationState::Closed:
+		case ETimecodeProviderSynchronizationState::Error:
+			return FLinearColor::Red;
+		case ETimecodeProviderSynchronizationState::Synchronized:
+			return FLinearColor::Green;
+		case ETimecodeProviderSynchronizationState::Synchronizing:
+			return FLinearColor::Yellow;
+		}
+	}
+	return FSlateColor::UseForeground();
+}
+
+FText STimecodeProvider::HandleStateText() const
+{
+	if (const UTimecodeProvider* TimecodeProviderPtr = GetTimecodeProvider())
+	{
+		ETimecodeProviderSynchronizationState State = TimecodeProviderPtr->GetSynchronizationState();
+		switch (State)
+		{
+		case ETimecodeProviderSynchronizationState::Error:
+		case ETimecodeProviderSynchronizationState::Closed:
+			return FEditorFontGlyphs::Ban;
+		case ETimecodeProviderSynchronizationState::Synchronized:
+			return FEditorFontGlyphs::Film;
+		case ETimecodeProviderSynchronizationState::Synchronizing:
+			return FEditorFontGlyphs::Hourglass_O;
+		}
+	}
+
+	return FEditorFontGlyphs::Exclamation;
+}
diff --git a/Engine/Plugins/Experimental/VirtualProductionUtilities/Source/VPUtilitiesEditor/Private/STimecodeProviderTab.cpp b/Engine/Plugins/Experimental/VirtualProductionUtilities/Source/VPUtilitiesEditor/Private/STimecodeProviderTab.cpp
new file mode 100644
--- /dev/null
+++ b/Engine/Plugins/Experimental/VirtualProductionUtilities/Source/VPUtilitiesEditor/Private/STimecodeProviderTab.cpp
@@ -0,0 +1,191 @@
+// Copyright Epic Games, Inc. All Rights Reserved.
+
+#include "STimecodeProviderTab.h"
+
+#include "EditorStyleSet.h"
+#include "Engine/Engine.h"
+#include "Engine/TimecodeProvider.h"
+#include "Framework/Application/SlateApplication.h"
+#include "Framework/Docking/TabManager.h"
+#include "Framework/MultiBox/MultiBoxBuilder.h"
+#include "Framework/MultiBox/SToolBarComboButtonBlock.h"
+#include "LevelEditor.h"
+#include "Modules/ModuleManager.h"
+#include "ObjectEditorUtils.h"
+#include "ScopedTransaction.h"
+#include "STimecodeProvider.h"
+#include "Widgets/Docking/SDockTab.h"
+#include "Widgets/Input/SComboButton.h"
+#include "Widgets/Input/SSpinBox.h"
+#include "WorkspaceMenuStructure.h"
+#include "WorkspaceMenuStructureModule.h"
+
+
+#define LOCTEXT_NAMESPACE "TimecodeProviderTab"
+
+namespace TimecodeProviderTab
+{
+	static FDelegateHandle LevelEditorTabManagerChangedHandle;
+	static const FName NAME_TimecodeProviderTab = FName("TimecodeProviderTab");
+
+	TSharedRef<SDockTab> CreateTab(const FSpawnTabArgs& Args)
+	{
+		return SNew(SDockTab)
+			.TabRole(ETabRole::NomadTab)
+			[
+				SNew(STimecodeProviderTab)
+			];
+	}
+}
+
+void STimecodeProviderTab::RegisterNomadTabSpawner()
+{
+	auto RegisterTabSpawner = []()
+	{
+		FLevelEditorModule& LevelEditorModule = FModuleManager::LoadModuleChecked<FLevelEditorModule>("LevelEditor");
+		TSharedPtr<FTabManager> LevelEditorTabManager = LevelEditorModule.GetLevelEditorTabManager();
+
+		LevelEditorTabManager->RegisterTabSpawner(TimecodeProviderTab::NAME_TimecodeProviderTab, FOnSpawnTab::CreateStatic(&TimecodeProviderTab::CreateTab))
+			.SetDisplayName(NSLOCTEXT("TimecodeProviderTab", "DisplayName", "Timecode Provider"))
+			.SetTooltipText(NSLOCTEXT("TimecodeProviderTab", "TooltipText", "Displays the Timecode and the state of the current Timecode Provider."))
+			.SetGroup(WorkspaceMenu::GetMenuStructure().GetLevelEditorVirtualProductionCategory())
+			.SetIcon(FSlateIcon(FEditorStyle::GetStyleSetName(), "TimecodeProvider.TabIcon"));
+	};
+
+	FLevelEditorModule& LevelEditorModule = FModuleManager::LoadModuleChecked<FLevelEditorModule>("LevelEditor");
+	if (LevelEditorModule.GetLevelEditorTabManager())
+	{
+		RegisterTabSpawner();
+	}
+	else
+	{
+		TimecodeProviderTab::LevelEditorTabManagerChangedHandle = LevelEditorModule.OnTabManagerChanged().AddLambda(RegisterTabSpawner);
+	}
+}
+
+void STimecodeProviderTab::UnregisterNomadTabSpawner()
+{
+	if (FSlateApplication::IsInitialized() && FModuleManager::Get().IsModuleLoaded("LevelEditor"))
+	{
+		FLevelEditorModule* LevelEditorModule = FModuleManager::GetModulePtr<FLevelEditorModule>("LevelEditor");
+		TSharedPtr<FTabManager> LevelEditorTabManager;
+		if (LevelEditorModule)
+		{
+			LevelEditorTabManager = LevelEditorModule->GetLevelEditorTabManager();
+			LevelEditorModule->OnTabManagerChanged().Remove(TimecodeProviderTab::LevelEditorTabManagerChangedHandle);
+		}
+
+		if (LevelEditorTabManager.IsValid())
+		{
+			LevelEditorTabManager->UnregisterTabSpawner(TimecodeProviderTab::NAME_TimecodeProviderTab);
+		}
+	}
+}
+
+void STimecodeProviderTab::Construct(const FArguments& InArgs)
+{
+	TSharedRef< SWidget > ButtonContent = SNew(SComboButton)
+		.ContentPadding(0)
+		.ButtonStyle(&FCoreStyle::Get(), "ToolBar.Button")
+		.ForegroundColor(FCoreStyle::Get().GetSlateColor("DefaultForeground"))
+		.ButtonContent()
+		[
+			SNullWidget::NullWidget
+		]
+		.OnGetMenuContent(this, &STimecodeProviderTab::OnGetMenuContent);
+
+	ButtonContent->SetEnabled(MakeAttributeLambda([] { return GEngine && GEngine->GetTimecodeProvider() != nullptr; }));
+
+	ChildSlot
+	[
+		SNew(SBorder)
+		.Padding(FMargin(0, 3, 0, 0))
+		.BorderImage(FEditorStyle::GetBrush("NoBorder"))
+		.IsEnabled(FSlateApplication::Get().GetNormalExecutionAttribute())
+		[
+			SNew(SBorder)
+			.BorderImage(FEditorStyle::GetBrush("ToolPanel.GroupBorder"))
+			[
+				SNew(SVerticalBox)
+				+ SVerticalBox::Slot()
+				.AutoHeight()
+				[
+					SNew(SHorizontalBox)
+					+ SHorizontalBox::Slot()
+					.AutoWidth()
+					.Padding(4, -4, 2, 0)
+					[
+						SNew(STimecodeProvider)
+					]
+					+ SHorizontalBox::Slot()
+					.AutoWidth()
+					.Padding(2, 0, 2, 0)
+					[
+						ButtonContent
+					]
+				]
+			]
+		]
+	];
+}
+
+TSharedRef<SWidget> STimecodeProviderTab::OnGetMenuContent()
+{
+	const bool bShouldCloseWindowAfterMenuSelection = true;
+	FMenuBuilder MenuBuilder(bShouldCloseWindowAfterMenuSelection, nullptr);
+
+	check(GEngine);
+	if (GEngine->GetTimecodeProvider())
+	{
+		MenuBuilder.BeginSection(TEXT("TimecodeProvider"), LOCTEXT("TimecodeProvider", "Timecode Provider"));
+
+		MenuBuilder.AddMenuEntry(
+			LOCTEXT("ReapplyMenuLabel", "Reinitialize"),
+			LOCTEXT("ReapplyMenuToolTip", "Reinitialize the current Timecode Provider."),
+			FSlateIcon(),
+			FUIAction(FExecuteAction::CreateUObject(GEngine, &UEngine::ReinitializeTimecodeProvider))
+			);
+
+		MenuBuilder.EndSection();
+	}
+
+	if (GEngine->GetTimecodeProvider())
+	{
+		MenuBuilder.BeginSection("Settings", LOCTEXT("Settings", "Settings"));
+		{
+			TSharedRef<SWidget> RefreshDelay = SNew(SSpinBox<float>)
+				.ToolTipText(LOCTEXT("FrameDelay_ToolTip", "Number of frames to subtract from the original timecode."))
+				.Value(this, &STimecodeProviderTab::GetFrameDelay)
+				.MinValue(TOptional<float>())
+				.MaxValue(TOptional<float>())
+				.OnValueCommitted(this, &STimecodeProviderTab::SetFrameDelay);
+
+			MenuBuilder.AddWidget(RefreshDelay, LOCTEXT("FrameDelay", "Frame Delay"));
+		}
+		MenuBuilder.EndSection();
+	}
+
+	return MenuBuilder.MakeWidget();
+}
+
+float STimecodeProviderTab::GetFrameDelay() const
+{
+	if (UTimecodeProvider* TimecodeProviderPtr = GEngine->GetTimecodeProvider())
+	{
+		return TimecodeProviderPtr->FrameDelay;
+	}
+	return 0.f;
+}
+
+void STimecodeProviderTab::SetFrameDelay(float InNewValue, ETextCommit::Type)
+{
+	if (UTimecodeProvider* TimecodeProviderPtr = GEngine->GetTimecodeProvider())
+	{
+		if (TimecodeProviderPtr->FrameDelay != InNewValue)
+		{
+			FObjectEditorUtils::SetPropertyValue(TimecodeProviderPtr, GET_MEMBER_NAME_CHECKED(UTimecodeProvider, FrameDelay), InNewValue);
+		}
+	}
+}
+
+#undef LOCTEXT_NAMESPACE
diff --git a/Engine/Plugins/Experimental/VirtualProductionUtilities/Source/VPUtilitiesEditor/Private/STimecodeProviderTab.h b/Engine/Plugins/Experimental/VirtualProductionUtilities/Source/VPUtilitiesEditor/Private/STimecodeProviderTab.h
new file mode 100644
--- /dev/null
+++ b/Engine/Plugins/Experimental/VirtualProductionUtilities/Source/VPUtilitiesEditor/Private/STimecodeProviderTab.h
@@ -0,0 +1,34 @@
+// Copyright Epic Games, Inc. All Rights Reserved.
+
+#pragma once
+
+#include "CoreMinimal.h"
+#include "Misc/FrameRate.h"
+#include "Types/SlateEnums.h"
+#include "Widgets/DeclarativeSyntaxSupport.h"
+#include "Widgets/SCompoundWidget.h"
+
+class SWidget;
+
+/*
+ * STimecodeProviderTab
+ */
+class STimecodeProviderTab : public SCompoundWidget
+{
+public:
+	static void RegisterNomadTabSpawner();
+	static void UnregisterNomadTabSpawner();
+
+public:
+	SLATE_BEGIN_ARGS(STimecodeProviderTab) {}
+	SLATE_END_ARGS()
+
+	void Construct(const FArguments& InArgs);
+
+private:
+	float GetFrameDelay() const;
+	void SetFrameDelay(float, ETextCommit::Type);
+
+private:
+	TSharedRef<SWidget> OnGetMenuContent();
+};
diff --git a/Engine/Plugins/Experimental/VirtualProductionUtilities/Source/VPUtilitiesEditor/Private/VPUtilitiesEditorModule.cpp b/Engine/Plugins/Experimental/VirtualProductionUtilities/Source/VPUtilitiesEditor/Private/VPUtilitiesEditorModule.cpp
--- a/Engine/Plugins/Experimental/VirtualProductionUtilities/Source/VPUtilitiesEditor/Private/VPUtilitiesEditorModule.cpp
+++ b/Engine/Plugins/Experimental/VirtualProductionUtilities/Source/VPUtilitiesEditor/Private/VPUtilitiesEditorModule.cpp
@@ -16,6 +16,7 @@
 #include "OSCManager.h"
 #include "OSCServer.h"
 #include "SGenlockProviderTab.h"
+#include "STimecodeProviderTab.h"
 #include "Textures/SlateIcon.h"
 #include "VPSettings.h"
 #include "VPUtilitiesEditorSettings.h"
@@ -39,6 +40,7 @@ void FVPUtilitiesEditorModule::StartupModule()
 	CustomUIHandler->Init();
 
 	SGenlockProviderTab::RegisterNomadTabSpawner(WorkspaceMenu::GetMenuStructure().GetLevelEditorVirtualProductionCategory());
+	STimecodeProviderTab::RegisterNomadTabSpawner();
 
 	RegisterSettings();
 
@@ -51,6 +53,7 @@ void FVPUtilitiesEditorModule::StartupModule()
 void FVPUtilitiesEditorModule::ShutdownModule()
 {
 	UnregisterSettings();
+	STimecodeProviderTab::UnregisterNomadTabSpawner();
 	SGenlockProviderTab::UnregisterNomadTabSpawner();
 
 	if (UObjectInitialized())
diff --git a/Engine/Plugins/Experimental/VirtualProductionUtilities/Source/VPUtilitiesEditor/Public/STimecodeProvider.h b/Engine/Plugins/Experimental/VirtualProductionUtilities/Source/VPUtilitiesEditor/Public/STimecodeProvider.h
new file mode 100644
--- /dev/null
+++ b/Engine/Plugins/Experimental/VirtualProductionUtilities/Source/VPUtilitiesEditor/Public/STimecodeProvider.h
@@ -0,0 +1,80 @@
+// Copyright Epic Games, Inc. All Rights Reserved.
+
+#pragma once
+
+#include "CoreMinimal.h"
+#include "Widgets/DeclarativeSyntaxSupport.h"
+#include "Widgets/SCompoundWidget.h"
+
+#include "Engine/TimecodeProvider.h"
+#include "Fonts/SlateFontInfo.h"
+#include "Layout/Visibility.h"
+#include "Misc/Attribute.h"
+#include "Misc/Timecode.h"
+#include "Styling/CoreStyle.h"
+#include "Styling/SlateColor.h"
+#include "UObject/WeakObjectPtrTemplates.h"
+
+struct FSlateBrush;
+
+class VPUTILITIESEDITOR_API STimecodeProvider : public SCompoundWidget
+{	
+public:
+	SLATE_BEGIN_ARGS(STimecodeProvider)
+		: _TimecodeProviderColor(FLinearColor::Gray)
+		, _DisplayFrameRate(true)
+		, _DisplaySynchronizationState(true)
+		, _TimecodeColor(FLinearColor::White)
+		, _DisplayLabel(false)
+		, _LabelColor(FLinearColor::Gray)
+	{
+		FSlateFontInfo NormalFontInfo = FCoreStyle::Get().GetFontStyle(TEXT("NormalText"));
+		int32 OriginalSize = NormalFontInfo.Size;
+
+		_LabelFont = NormalFontInfo;
+		NormalFontInfo.Size = OriginalSize - 4;
+		_TimecodeProviderFont = NormalFontInfo;
+		NormalFontInfo.Size = OriginalSize + 16;
+		_TimecodeFont = NormalFontInfo;
+	}
+		/** The font for this TimecodeProvider text */
+		SLATE_ATTRIBUTE(FSlateFontInfo, TimecodeProviderFont)
+		/** The color for this TimecodeProvider text */
+		SLATE_ATTRIBUTE(FSlateColor, TimecodeProviderColor)
+		/** Should display the TimecodeProvider's frame rate */
+		SLATE_ARGUMENT(bool, DisplayFrameRate)
+		/** Should display the TimecodeProvider's synchronization state */
+		SLATE_ARGUMENT(bool, DisplaySynchronizationState)
+		/** Override the Timecode Provider to display */
+		SLATE_ATTRIBUTE(TWeakObjectPtr<UTimecodeProvider>, OverrideTimecodeProvider)
+
+		/** The font for the timecode text */
+		SLATE_ATTRIBUTE(FSlateFontInfo, TimecodeFont)
+		/** The color for the timecode text */
+		SLATE_ATTRIBUTE(FSlateColor, TimecodeColor)
+
+		/** Should display the label (hours, mins, secs, frames) */
+		SLATE_ATTRIBUTE(bool, DisplayLabel)
+		/** The font for this label text */
+		SLATE_ATTRIBUTE(FSlateFontInfo, LabelFont)
+		/** The color for this label text */
+		SLATE_ATTRIBUTE(FSlateColor, LabelColor)
+	SLATE_END_ARGS()
+
+	/**
+	 * Construct this widget
+	 *
+	 * @param	InArgs	The declaration data for this widget
+	 */
+	void Construct(const FArguments& InArgs);
+
+private:
+
+	const UTimecodeProvider* GetTimecodeProvider() const;
+	FSlateColor HandleIconColorAndOpacity() const;
+	FText HandleStateText() const;
+
+
+private:
+	TAttribute<TWeakObjectPtr<UTimecodeProvider>> OverrideTimecodeProvider;
+};
diff --git a/Engine/Plugins/Experimental/VirtualProductionUtilities/Source/VPUtilitiesEditor/VPUtilitiesEditor.Build.cs b/Engine/Plugins/Experimental/VirtualProductionUtilities/Source/VPUtilitiesEditor/VPUtilitiesEditor.Build.cs
--- a/Engine/Plugins/Experimental/VirtualProductionUtilities/Source/VPUtilitiesEditor/VPUtilitiesEditor.Build.cs
+++ b/Engine/Plugins/Experimental/VirtualProductionUtilities/Source/VPUtilitiesEditor/VPUtilitiesEditor.Build.cs
@@ -15,6 +15,9 @@ public class VPUtilitiesEditor : ModuleRules
 				"Core",
 				"CoreUObject",
 				"EditorSubsystem",
+				"Engine",
+				"Slate",
+				"SlateCore",
 				"VPUtilities",                
 				"VREditor",
             }
@@ -23,23 +26,21 @@ public class VPUtilitiesEditor : ModuleRules
 		PrivateDependencyModuleNames.AddRange(
 			new string[]
 			{
+                "CinematicCamera",
 				"EditorFramework",
 				"EditorStyle",
-				"Engine",
 				"GameplayTags",
+				"InputCore",
 				"LevelEditor",
+				"OSC",
+				"PlacementMode",
 				"Settings",
-				"Slate",
-				"SlateCore",
 				"TimeManagement",
 				"UMG",
 				"UMGEditor",
 				"UnrealEd",
 				"VPBookmark",
 				"WorkspaceMenuStructure",
-                "CinematicCamera",
-				"OSC",
-				"PlacementMode"
             }
 		);
 	}
diff --git a/Engine/Plugins/VirtualProduction/TimedDataMonitor/Source/TimedDataMonitorEditor/TimedDataMonitorEditor.build.cs b/Engine/Plugins/VirtualProduction/TimedDataMonitor/Source/TimedDataMonitorEditor/TimedDataMonitorEditor.build.cs
--- a/Engine/Plugins/VirtualProduction/TimedDataMonitor/Source/TimedDataMonitorEditor/TimedDataMonitorEditor.build.cs
+++ b/Engine/Plugins/VirtualProduction/TimedDataMonitor/Source/TimedDataMonitorEditor/TimedDataMonitorEditor.build.cs
@@ -21,7 +21,6 @@ public class TimedDataMonitorEditor : ModuleRules
 				"SlateCore",
 				"TimedDataMonitor",
 				"TimeManagement",
-				"TimeManagementEditor",
 				"UnrealEd",
 				"WorkspaceMenuStructure",
 			});
diff --git a/Engine/Source/Editor/TimeManagementEditor/Private/STimecodeProvider.cpp b/Engine/Source/Editor/TimeManagementEditor/Private/STimecodeProvider.cpp
deleted file mode 100644
diff --git a/Engine/Source/Editor/TimeManagementEditor/Private/STimecodeProviderTab.cpp b/Engine/Source/Editor/TimeManagementEditor/Private/STimecodeProviderTab.cpp
deleted file mode 100644
diff --git a/Engine/Source/Editor/TimeManagementEditor/Private/STimecodeProviderTab.h b/Engine/Source/Editor/TimeManagementEditor/Private/STimecodeProviderTab.h
deleted file mode 100644
diff --git a/Engine/Source/Editor/TimeManagementEditor/Private/TimeManagementEditorModule.cpp b/Engine/Source/Editor/TimeManagementEditor/Private/TimeManagementEditorModule.cpp
deleted file mode 100644
diff --git a/Engine/Source/Editor/TimeManagementEditor/Public/STimecodeProvider.h b/Engine/Source/Editor/TimeManagementEditor/Public/STimecodeProvider.h
deleted file mode 100644
diff --git a/Engine/Source/Editor/TimeManagementEditor/TimeManagementEditor.Build.cs b/Engine/Source/Editor/TimeManagementEditor/TimeManagementEditor.Build.cs
deleted file mode 100644
diff --git a/Engine/Source/Editor/UnrealEd/UnrealEd.Build.cs b/Engine/Source/Editor/UnrealEd/UnrealEd.Build.cs
--- a/Engine/Source/Editor/UnrealEd/UnrealEd.Build.cs
+++ b/Engine/Source/Editor/UnrealEd/UnrealEd.Build.cs
@@ -254,7 +254,6 @@ public class UnrealEd : ModuleRules
 				"OverlayEditor",
 				"ClothPainter",
 				"Media",
-				"TimeManagementEditor",
 				"VirtualTexturingEditor",
 				"TraceInsights",
 			}
