-Fixing VPFullScreenWidget crash when loading a new level

#UE-121925

diff --git a/Engine/Plugins/Experimental/VirtualProductionUtilities/Source/VPUtilities/Private/VPFullScreenUserWidget.cpp b/Engine/Plugins/Experimental/VirtualProductionUtilities/Source/VPUtilities/Private/VPFullScreenUserWidget.cpp
--- a/Engine/Plugins/Experimental/VirtualProductionUtilities/Source/VPUtilities/Private/VPFullScreenUserWidget.cpp
+++ b/Engine/Plugins/Experimental/VirtualProductionUtilities/Source/VPUtilities/Private/VPFullScreenUserWidget.cpp
@@ -642,6 +642,7 @@ bool UVPFullScreenUserWidget::Display(UWorld* InWorld)
 		if (bWasAdded)
 		{
 			FWorldDelegates::LevelRemovedFromWorld.AddUObject(this, &UVPFullScreenUserWidget::OnLevelRemovedFromWorld);
+			FWorldDelegates::OnWorldCleanup.AddUObject(this, &UVPFullScreenUserWidget::OnWorldCleanup);
 
 			// If we are using Composure as our output, then send the WidgetRenderTarget to each one
 			if (CurrentDisplayType == EVPWidgetDisplayType::Composure)
@@ -688,6 +689,7 @@ void UVPFullScreenUserWidget::Hide()
 	{
 		ReleaseWidget();
 		FWorldDelegates::LevelRemovedFromWorld.RemoveAll(this);
+		FWorldDelegates::OnWorldCleanup.RemoveAll(this);
 
 		if (CurrentDisplayType == EVPWidgetDisplayType::Viewport)
 		{
@@ -762,6 +764,14 @@ void UVPFullScreenUserWidget::OnLevelRemovedFromWorld(ULevel* InLevel, UWorld* I
 	}
 }
 
+void UVPFullScreenUserWidget::OnWorldCleanup(UWorld* InWorld, bool bSessionEnded, bool bCleanupResources)
+{
+	if (IsDisplayed() && World == InWorld)
+	{
+		Hide();
+	}
+}
+
 FVector2D UVPFullScreenUserWidget::FindSceneViewportSize()
 {
 	FVector2D OutSize;
diff --git a/Engine/Plugins/Experimental/VirtualProductionUtilities/Source/VPUtilities/Public/VPFullScreenUserWidget.h b/Engine/Plugins/Experimental/VirtualProductionUtilities/Source/VPUtilities/Public/VPFullScreenUserWidget.h
--- a/Engine/Plugins/Experimental/VirtualProductionUtilities/Source/VPUtilities/Public/VPFullScreenUserWidget.h
+++ b/Engine/Plugins/Experimental/VirtualProductionUtilities/Source/VPUtilities/Public/VPFullScreenUserWidget.h
@@ -218,6 +218,7 @@ protected:
 
 private:
 	void OnLevelRemovedFromWorld(ULevel* InLevel, UWorld* InWorld);
+	void OnWorldCleanup(UWorld* InWorld, bool bSessionEnded, bool bCleanupResources);
 
 protected:
 	/** The display type when the world is an editor world. */
