Dynamic Material Instance

CL solves 3 issues

1. It adding dynamic material support to material details customization panel
2. It adding FOnGenerateGlobalRowExtension to RowGenerator
3. It adding ability to expose a entire material

diff --git a/Engine/Plugins/Experimental/VirtualProductionUtilities/Source/VPMaterialsEditor/Private/SMaterialDynamicParametersPanelWidget.cpp b/Engine/Plugins/Experimental/VirtualProductionUtilities/Source/VPMaterialsEditor/Private/SMaterialDynamicParametersPanelWidget.cpp
new file mode 100644
--- /dev/null
+++ b/Engine/Plugins/Experimental/VirtualProductionUtilities/Source/VPMaterialsEditor/Private/SMaterialDynamicParametersPanelWidget.cpp
@@ -0,0 +1,726 @@
+// Copyright Epic Games, Inc. All Rights Reserved.
+
+#include "SMaterialDynamicParametersPanelWidget.h"
+
+#include "PropertyHandle.h"
+
+#include "DetailWidgetRow.h"
+#include "EditorStyleSet.h"
+#include "IDetailPropertyRow.h"
+#include "IDetailTreeNode.h"
+#include "IPropertyRowGenerator.h"
+#include "PropertyCustomizationHelpers.h"
+#include "PropertyEditorModule.h"
+#include "Framework/MultiBox/MultiBoxBuilder.h"
+#include "Framework/MultiBox/MultiBoxDefs.h"
+#include "Materials/MaterialInstance.h"
+#include "Modules/ModuleManager.h"
+#include "UObject/Object.h"
+#include "Widgets/Layout/SWidgetSwitcher.h"
+#include "Widgets/Text/SInlineEditableTextBlock.h"
+#include "Widgets/Views/SExpanderArrow.h"
+#include "Widgets/Views/STreeView.h"
+
+#define LOCTEXT_NAMESPACE "MaterialDynamicParametersPanelWidget"
+
+enum EStackDataType
+{
+	Stack,
+	Asset,
+	Group,
+	Property,
+	PropertyChild,
+};
+
+struct FPropertySortedParamData
+{
+	EStackDataType StackDataType;
+
+	FMaterialParameterInfo ParameterInfo;
+
+	TSharedPtr<IDetailTreeNode> ParameterNode;
+
+	TSharedPtr<IPropertyHandle> ParameterHandle;
+
+	TArray<TSharedPtr<FPropertySortedParamData>> Children;
+
+	FString NodeKey;
+
+	FName GroupName;
+};
+
+struct FPropertyUnsortedParamData
+{
+	TSharedPtr<IDetailTreeNode> ParameterNode;
+
+	TSharedPtr<IPropertyHandle> ParameterHandle;
+
+	FMaterialParameterInfo ParameterInfo;
+
+	FName GroupName;
+};
+
+namespace Utilities
+{
+	template<typename ParameterValueType>
+	void AddParameters(TSharedPtr<IPropertyHandle> InPropertyHandle, const FName InParameterName, TArray<TSharedPtr<IPropertyHandle>>& OutDeferredSearches, TArray<FPropertyUnsortedParamData>& OutUnsortedParameters)
+	{
+		if (!InPropertyHandle.IsValid())
+		{
+			return;
+		}
+		
+		TArray<void*> GroupPtrs;
+		InPropertyHandle->AccessRawData(GroupPtrs);
+		const auto GroupIt = GroupPtrs.CreateConstIterator();
+		TArray<ParameterValueType>* ParameterGroupArrayPtr = reinterpret_cast<TArray<ParameterValueType>*>(*GroupIt);
+		if (!ParameterGroupArrayPtr)
+		{
+			return;
+		}
+
+		const TSharedPtr<IPropertyHandleArray> ParameterValuesArray = InPropertyHandle->AsArray();
+
+		uint32 NumElements = 0;
+		ParameterValuesArray->GetNumElements(NumElements);
+		for (uint32 Index = 0; Index < NumElements; ++Index)
+		{
+			const ParameterValueType& ParameterValue = (*ParameterGroupArrayPtr)[Index];
+			if (!ParameterValue.ExpressionGUID.IsValid())
+			{
+				continue;
+			}
+
+			if (const TSharedPtr<IPropertyHandle> ElementPropertyHandle = ParameterValuesArray->GetElement(Index))
+			{
+				FPropertyUnsortedParamData NonLayerProperty;
+				TSharedPtr<IPropertyHandle> ParameterValuePropertyHandle = ElementPropertyHandle->GetChildHandle(InParameterName);
+
+				NonLayerProperty.ParameterInfo = ParameterValue.ParameterInfo;
+				NonLayerProperty.GroupName = TEXT("Group");
+
+				OutDeferredSearches.Add(ParameterValuePropertyHandle);
+				OutUnsortedParameters.Add(NonLayerProperty);
+			}
+		}
+	}
+};
+
+// ********* SMaterialParametersOverviewTree *******
+class SMaterialDynamicParametersOverviewTree : public STreeView<TSharedPtr<FPropertySortedParamData>>
+{
+	friend class SMaterialDynamicParametersOverviewTreeItem;
+
+public:
+
+	SLATE_BEGIN_ARGS(SMaterialDynamicParametersOverviewTree)
+	{}
+
+	SLATE_ARGUMENT(UMaterialInstance*, InMaterialInstance)
+	SLATE_ARGUMENT(TSharedPtr<SMaterialDynamicParametersPanelWidget>, InOwner)
+	SLATE_ARGUMENT(TSharedPtr<SScrollBar>, InScrollbar)
+	SLATE_END_ARGS()
+
+	/** Constructs this widget with InArgs */
+	void Construct(const FArguments& InArgs);
+	TSharedRef< ITableRow > OnGenerateRowMaterialLayersFunctionsTreeView(TSharedPtr<FPropertySortedParamData> Item, const TSharedRef< STableViewBase >& OwnerTable);
+	void OnGetChildrenMaterialLayersFunctionsTreeView(TSharedPtr<FPropertySortedParamData> InParent, TArray< TSharedPtr<FPropertySortedParamData> >& OutChildren);
+	void OnExpansionChanged(TSharedPtr<FPropertySortedParamData> Item, bool bIsExpanded);
+	void SetParentsExpansionState();
+
+	/** Builds the custom parameter groups category */
+	void CreateGroupsWidget();
+
+	TWeakPtr<SMaterialDynamicParametersPanelWidget> GetOwner() { return Owner; }
+	bool HasAnyParameters() const { return bHasAnyParameters; }
+
+	FDetailColumnSizeData& GetColumnSizeData() { return ColumnSizeData; }
+
+protected:
+
+	void ShowSubParameters();
+
+public:
+	/** Object that stores all of the possible parameters we can edit */
+	TWeakObjectPtr<UMaterialInstance> MaterialInstance;
+
+private:
+	TArray<TSharedPtr<FPropertySortedParamData>> SortedParameters;
+
+	TArray<FPropertyUnsortedParamData> UnsortedParameters;
+
+	TWeakPtr<SMaterialDynamicParametersPanelWidget> Owner;
+
+	bool bHasAnyParameters;
+	
+	FDetailColumnSizeData ColumnSizeData;
+};
+
+// ********* SMaterialParametersOverviewTreeItem *******
+class SMaterialDynamicParametersOverviewTreeItem : public STableRow< TSharedPtr<FPropertySortedParamData> >
+{
+public:
+
+	SLATE_BEGIN_ARGS(SMaterialDynamicParametersOverviewTreeItem)
+		: _StackParameterData(nullptr),
+		_MaterialInstance(nullptr)
+	{}
+
+	/** The item content. */
+	SLATE_ARGUMENT(TSharedPtr<FPropertySortedParamData>, StackParameterData)
+	SLATE_ARGUMENT(UMaterialInstance*, MaterialInstance)
+	SLATE_ARGUMENT(TSharedPtr<SMaterialDynamicParametersOverviewTree>, InTree)
+	SLATE_END_ARGS()
+
+	void RefreshOnRowChange(const FAssetData& AssetData, TSharedPtr<SMaterialDynamicParametersOverviewTree> InTree);
+
+	/**
+	* Construct the widget
+	*
+	* @param InArgs   A declaration from which to construct the widget
+	*/
+	void Construct(const FArguments& InArgs, const TSharedRef<STableViewBase>& InOwnerTableView);
+
+private:
+	const FSlateBrush* GetBorderImage() const;
+
+	FSlateColor GetOuterBackgroundColor(TSharedPtr<FPropertySortedParamData> InParamData) const;
+
+	TSharedRef<SWidget> GetRowExtensionButtons(TSharedPtr<IPropertyHandle> InPropertyHandle) const;
+private:
+
+	/** The node info to build the tree view row from. */
+	TSharedPtr<FPropertySortedParamData> StackParameterData;
+
+	/** The tree that contains this item */
+	TWeakPtr<SMaterialDynamicParametersOverviewTree> Tree;
+
+	/** The set of material parameters this is associated with */
+	TWeakObjectPtr<UMaterialInstance> MaterialInstance;
+};
+
+
+
+const FSlateBrush* SMaterialDynamicParametersOverviewTreeItem::GetBorderImage() const
+{
+	return FAppStyle::Get().GetBrush("DetailsView.CategoryMiddle");
+}
+
+FSlateColor SMaterialDynamicParametersOverviewTreeItem::GetOuterBackgroundColor(TSharedPtr<FPropertySortedParamData> InParamData) const
+{
+	if (IsHovered() || InParamData->StackDataType == EStackDataType::Group)
+	{
+		return FAppStyle::Get().GetSlateColor("Colors.Header");
+	}
+
+	return FAppStyle::Get().GetSlateColor("Colors.Panel");
+}
+
+TSharedRef<SWidget> SMaterialDynamicParametersOverviewTreeItem::GetRowExtensionButtons(TSharedPtr<IPropertyHandle> InPropertyHandle) const
+{
+	if (!InPropertyHandle.IsValid())
+	{
+		return SNullWidget::NullWidget;
+	}
+
+	const auto ResetDelegate = FExecuteAction::CreateLambda([InPropertyHandle]()
+	{
+		if (InPropertyHandle.IsValid())
+		{
+			InPropertyHandle->ResetToDefault();
+		}
+	});
+	
+	FPropertyEditorModule& PropertyEditorModule = FModuleManager::Get().GetModuleChecked<FPropertyEditorModule>("PropertyEditor");
+
+	TArray<FPropertyRowExtensionButton> ExtensionButtons;
+	FPropertyRowExtensionButton& ResetToDefault = ExtensionButtons.AddDefaulted_GetRef();
+	ResetToDefault.Label = NSLOCTEXT("PropertyEditor", "ResetToDefault", "Reset to Default");
+	ResetToDefault.UIAction = FUIAction(ResetDelegate);
+
+	ResetToDefault.Icon = FSlateIcon(FAppStyle::Get().GetStyleSetName(), "PropertyWindow.DiffersFromDefault");
+	ResetToDefault.ToolTip = NSLOCTEXT("PropertyEditor", "ResetToDefaultToolTip", "Reset this property to its default value.");
+
+	// Add custom property extensions
+	FOnGenerateGlobalRowExtensionArgs RowExtensionArgs;
+	RowExtensionArgs.PropertyHandle = InPropertyHandle;
+	PropertyEditorModule.GetGlobalRowExtensionDelegate().Broadcast(RowExtensionArgs, ExtensionButtons);
+	
+	// Build extension toolbar 
+	FSlimHorizontalToolBarBuilder ToolbarBuilder(TSharedPtr<FUICommandList>(), FMultiBoxCustomization::None);
+	ToolbarBuilder.SetLabelVisibility(EVisibility::Collapsed);
+	ToolbarBuilder.SetStyle(&FAppStyle::Get(), "DetailsView.ExtensionToolBar");
+	for (const FPropertyRowExtensionButton& Extension : ExtensionButtons)
+	{
+		ToolbarBuilder.AddToolBarButton(Extension.UIAction, NAME_None, Extension.Label, Extension.ToolTip, Extension.Icon);
+	}
+
+	return ToolbarBuilder.MakeWidget();
+}
+
+void SMaterialDynamicParametersOverviewTreeItem::RefreshOnRowChange(const FAssetData& AssetData, TSharedPtr<SMaterialDynamicParametersOverviewTree> InTree)
+{
+	if (InTree.IsValid())
+	{
+		InTree->CreateGroupsWidget();
+	}
+}
+
+void SMaterialDynamicParametersOverviewTreeItem::Construct(const FArguments& InArgs, const TSharedRef<STableViewBase>& InOwnerTableView)
+{
+	StackParameterData = InArgs._StackParameterData;
+	MaterialInstance = InArgs._MaterialInstance;
+	Tree = InArgs._InTree;
+
+	TSharedRef<SWidget> LeftSideWidget = SNullWidget::NullWidget;
+	TSharedRef<SWidget> RightSideWidget = SNullWidget::NullWidget;
+	FText NameOverride;
+	TSharedRef<SVerticalBox> WrapperWidget = SNew(SVerticalBox);
+
+	auto GenerateLeftAndRightWidgets = [&](IDetailTreeNode& InNode)
+	{
+		const FNodeWidgets NodeWidgets = InNode.CreateNodeWidgets();
+		LeftSideWidget = NodeWidgets.NameWidget.ToSharedRef();
+
+		RightSideWidget = SNew(SHorizontalBox)
+				+SHorizontalBox::Slot()
+				.FillWidth(1.f)
+				[
+					NodeWidgets.ValueWidget.ToSharedRef()
+				]
+				+ SHorizontalBox::Slot()
+				.AutoWidth()
+				[
+					GetRowExtensionButtons(InNode.CreatePropertyHandle())
+				];
+	};
+
+// GROUP --------------------------------------------------
+	if (StackParameterData->StackDataType == EStackDataType::Group)
+	{
+		NameOverride = FText::FromName(StackParameterData->GroupName);
+		LeftSideWidget = SNew(STextBlock)
+			.TransformPolicy(ETextTransformPolicy::ToUpper)
+			.Text(NameOverride)
+			.Font(FAppStyle::Get().GetFontStyle("PropertyWindow.BoldFont"))
+			.TextStyle(FAppStyle::Get(), "DetailsView.CategoryTextStyle");
+	}
+// END GROUP
+
+// PROPERTY ----------------------------------------------
+	bool bIsPaddedProperty = false;
+	if (StackParameterData->StackDataType == EStackDataType::Property)
+	{
+		NameOverride = FText::FromName(StackParameterData->ParameterInfo.Name);
+
+		IDetailTreeNode& Node = *StackParameterData->ParameterNode;
+		TSharedPtr<IDetailPropertyRow> GeneratedRow = StaticCastSharedPtr<IDetailPropertyRow>(Node.GetRow());
+		IDetailPropertyRow& Row = *GeneratedRow.Get();
+		Row.DisplayName(NameOverride);
+
+
+		if (TSharedPtr<IPropertyHandle> PropertyHandle = StackParameterData->ParameterNode->CreatePropertyHandle())
+		{
+			PropertyHandle->MarkResetToDefaultCustomized(true);
+		}
+
+		FDetailWidgetDecl* CustomNameWidget = Row.CustomNameWidget();
+		if (CustomNameWidget)
+		{
+			(*CustomNameWidget)
+			[
+				SNew(SHorizontalBox)
+				+SHorizontalBox::Slot()
+				.VAlign(VAlign_Center)
+				[
+					SNew(STextBlock)
+					.Text(NameOverride)
+					.Font(FEditorStyle::GetFontStyle(TEXT("PropertyWindow.NormalFont")))
+				]
+			];
+		}
+
+		bIsPaddedProperty = true;
+
+		GenerateLeftAndRightWidgets(Node);
+	}
+// END PROPERTY
+
+// PROPERTY CHILD ----------------------------------------
+	if (StackParameterData->StackDataType == EStackDataType::PropertyChild)
+	{
+		GenerateLeftAndRightWidgets(*StackParameterData->ParameterNode);
+	}
+// END PROPERTY CHILD
+
+// FINAL WRAPPER
+
+	{
+		FDetailColumnSizeData& ColumnSizeData = InArgs._InTree->GetColumnSizeData();
+
+		float ValuePadding = bIsPaddedProperty ? 20.0f : 0.0f;
+		WrapperWidget->AddSlot()
+			.AutoHeight()
+			[
+				SNew(SBorder)
+				.BorderImage(FAppStyle::Get().GetBrush("DetailsView.GridLine"))
+				.Padding(FMargin(0, 0, 0, 1))
+				[
+					SNew(SBorder)
+					.Padding(3.0f)
+					.BorderImage(this, &SMaterialDynamicParametersOverviewTreeItem::GetBorderImage)
+					.BorderBackgroundColor(this, &SMaterialDynamicParametersOverviewTreeItem::GetOuterBackgroundColor, StackParameterData)
+					[
+						SNew(SSplitter)
+						.Style(FEditorStyle::Get(), "DetailsView.Splitter")
+						.PhysicalSplitterHandleSize(1.0f)
+						.HitDetectionSplitterHandleSize(5.0f)
+						+ SSplitter::Slot()
+						.Value(ColumnSizeData.NameColumnWidth)
+						.OnSlotResized(ColumnSizeData.OnNameColumnResized)
+						.Value(0.25f)
+						[
+							SNew(SHorizontalBox)
+							+ SHorizontalBox::Slot()
+							.AutoWidth()
+							.VAlign(VAlign_Center)
+							.Padding(FMargin(3.0f))
+							[
+								SNew(SExpanderArrow, SharedThis(this))
+							]
+							+ SHorizontalBox::Slot()
+							.Padding(FMargin(2.0f))
+							.VAlign(VAlign_Center)
+							[
+								LeftSideWidget
+							]
+						]
+						+ SSplitter::Slot()
+						.Value(ColumnSizeData.ValueColumnWidth)
+						.OnSlotResized(ColumnSizeData.OnValueColumnResized) 
+						[
+							SNew(SHorizontalBox)
+							+ SHorizontalBox::Slot()
+							.Padding(FMargin(5.0f, 2.0f, 5.0, 2.0f))
+							.HAlign(HAlign_Left)
+							.VAlign(VAlign_Center)
+							[
+								RightSideWidget
+							]
+						]
+					]
+				]
+			];
+	}
+
+	ChildSlot
+		[
+			WrapperWidget
+		];
+
+
+	STableRow< TSharedPtr<FPropertySortedParamData> >::ConstructInternal(
+		STableRow::FArguments()
+		.Style(FEditorStyle::Get(), "DetailsView.TreeView.TableRow")
+		.ShowSelection(false),
+		InOwnerTableView
+	);
+}
+
+void SMaterialDynamicParametersOverviewTree::Construct(const FArguments& InArgs)
+{
+	bHasAnyParameters = false;
+	MaterialInstance = InArgs._InMaterialInstance;
+	Owner = InArgs._InOwner;
+	CreateGroupsWidget();
+
+	STreeView<TSharedPtr<FPropertySortedParamData>>::Construct(
+		STreeView::FArguments()
+		.TreeItemsSource(&SortedParameters)
+		.SelectionMode(ESelectionMode::None)
+		.OnGenerateRow(this, &SMaterialDynamicParametersOverviewTree::OnGenerateRowMaterialLayersFunctionsTreeView)
+		.OnGetChildren(this, &SMaterialDynamicParametersOverviewTree::OnGetChildrenMaterialLayersFunctionsTreeView)
+		.OnExpansionChanged(this, &SMaterialDynamicParametersOverviewTree::OnExpansionChanged)
+		.ExternalScrollbar(InArgs._InScrollbar)
+	);
+}
+
+TSharedRef< ITableRow > SMaterialDynamicParametersOverviewTree::OnGenerateRowMaterialLayersFunctionsTreeView(TSharedPtr<FPropertySortedParamData> Item, const TSharedRef< STableViewBase >& OwnerTable)
+{
+	TSharedRef< SMaterialDynamicParametersOverviewTreeItem > ReturnRow = SNew(SMaterialDynamicParametersOverviewTreeItem, OwnerTable)
+		.StackParameterData(Item)
+		.MaterialInstance(MaterialInstance.Get())
+		.InTree(SharedThis(this));
+	return ReturnRow;
+}
+
+void SMaterialDynamicParametersOverviewTree::OnGetChildrenMaterialLayersFunctionsTreeView(TSharedPtr<FPropertySortedParamData> InParent, TArray< TSharedPtr<FPropertySortedParamData> >& OutChildren)
+{
+	OutChildren = InParent->Children;
+}
+
+void SMaterialDynamicParametersOverviewTree::OnExpansionChanged(TSharedPtr<FPropertySortedParamData> Item, bool bIsExpanded)
+{
+	bool* ExpansionValue = MaterialInstance->ParameterOverviewExpansion.Find(Item->NodeKey);
+	if (ExpansionValue == nullptr)
+	{
+		MaterialInstance->ParameterOverviewExpansion.Add(Item->NodeKey, bIsExpanded);
+	}
+	else if (*ExpansionValue != bIsExpanded)
+	{
+		MaterialInstance->ParameterOverviewExpansion.Emplace(Item->NodeKey, bIsExpanded);
+	}
+	// Expand any children that are also expanded
+	for (const TSharedPtr<FPropertySortedParamData> Child : Item->Children)
+	{
+		bool* ChildExpansionValue = MaterialInstance->ParameterOverviewExpansion.Find(Child->NodeKey);
+		if (ChildExpansionValue != nullptr && *ChildExpansionValue == true)
+		{
+			SetItemExpansion(Child, true);
+		}
+	}
+}
+
+void SMaterialDynamicParametersOverviewTree::SetParentsExpansionState()
+{
+	for (const auto& Pair : SortedParameters)
+	{
+		if (Pair->Children.Num())
+		{
+			bool* bIsExpanded = MaterialInstance->ParameterOverviewExpansion.Find(Pair->NodeKey);
+			if (bIsExpanded)
+			{
+				SetItemExpansion(Pair, *bIsExpanded);
+			}
+			else
+			{
+				SetItemExpansion(Pair, true);
+			}
+		}
+	}
+}
+
+void SMaterialDynamicParametersOverviewTree::CreateGroupsWidget()
+{
+	check(MaterialInstance.IsValid());
+	UnsortedParameters.Reset();
+	SortedParameters.Reset();
+
+	const TArray<TSharedRef<IDetailTreeNode>> TestData = GetOwner().Pin()->Generator->GetRootTreeNodes();
+	if (TestData.Num() == 0)
+	{
+		return;
+	}
+	TSharedPtr<IDetailTreeNode> Category = TestData[0];
+	TSharedPtr<IDetailTreeNode> ParameterGroups;
+	TArray<TSharedRef<IDetailTreeNode>> Children;
+	Category->GetChildren(Children);
+
+	TArray<TSharedPtr<IPropertyHandle>> DeferredSearches;
+
+	for (int32 ChildIdx = 0; ChildIdx < Children.Num(); ChildIdx++)
+	{
+		TSharedPtr<IPropertyHandle> PropertyHandle = Children[ChildIdx]->CreatePropertyHandle();
+
+		bHasAnyParameters = true;
+
+		if (!PropertyHandle.IsValid() || !PropertyHandle->GetProperty())
+		{
+			continue;
+		}
+
+		if (PropertyHandle->GetProperty()->GetFName() == GET_MEMBER_NAME_CHECKED(UMaterialInstance, ScalarParameterValues))
+		{
+			Utilities::AddParameters<FScalarParameterValue>(PropertyHandle,
+			                                                 GET_MEMBER_NAME_CHECKED(
+				                                                 FScalarParameterValue, ParameterValue),
+			                                                 DeferredSearches, UnsortedParameters);
+		}
+		else if (PropertyHandle->GetProperty()->GetFName() == GET_MEMBER_NAME_CHECKED(UMaterialInstance, VectorParameterValues))
+		{
+			Utilities::AddParameters<FVectorParameterValue>(PropertyHandle,
+															GET_MEMBER_NAME_CHECKED(
+																FVectorParameterValue, ParameterValue),
+															DeferredSearches, UnsortedParameters);
+		}
+		else if (PropertyHandle->GetProperty()->GetFName() == GET_MEMBER_NAME_CHECKED(UMaterialInstance, TextureParameterValues))
+		{
+			Utilities::AddParameters<FTextureParameterValue>(PropertyHandle,
+															GET_MEMBER_NAME_CHECKED(
+																FTextureParameterValue, ParameterValue),
+															DeferredSearches, UnsortedParameters);
+		}
+		else if (PropertyHandle->GetProperty()->GetFName() == GET_MEMBER_NAME_CHECKED(UMaterialInstance, FontParameterValues))
+		{
+			Utilities::AddParameters<FFontParameterValue>(PropertyHandle,
+															GET_MEMBER_NAME_CHECKED(
+																FFontParameterValue, FontValue),
+															DeferredSearches, UnsortedParameters);
+		}
+	}
+
+	checkf(UnsortedParameters.Num() == DeferredSearches.Num(), TEXT("Internal inconsistency: number of node searches does not match the number of properties"));
+	TArray<TSharedPtr<IDetailTreeNode>> DeferredResults = GetOwner().Pin()->Generator->FindTreeNodes(DeferredSearches);
+	checkf(UnsortedParameters.Num() == DeferredResults.Num(), TEXT("Internal inconsistency: number of node search results does not match the number of properties"));
+
+	for (int Idx = 0, NumUnsorted = UnsortedParameters.Num(); Idx < NumUnsorted; ++Idx)
+	{
+		FPropertyUnsortedParamData& NonLayerProperty = UnsortedParameters[Idx];
+		NonLayerProperty.ParameterNode = DeferredResults[Idx];
+		NonLayerProperty.ParameterHandle = NonLayerProperty.ParameterNode->CreatePropertyHandle();
+	}
+
+	Children.Empty();
+
+	ShowSubParameters();
+	RequestTreeRefresh();
+	SetParentsExpansionState();
+}
+
+void SMaterialDynamicParametersOverviewTree::ShowSubParameters()
+{
+	for (FPropertyUnsortedParamData Property : UnsortedParameters)
+	{
+		TSharedPtr<FPropertySortedParamData> GroupProperty(new FPropertySortedParamData());
+		GroupProperty->ParameterInfo = Property.ParameterInfo;
+		GroupProperty->StackDataType = EStackDataType::Group;
+		GroupProperty->GroupName = Property.GroupName;
+		GroupProperty->NodeKey = FString::FromInt(GroupProperty->ParameterInfo.Index) + FString::FromInt(GroupProperty->ParameterInfo.Association) + Property.GroupName.ToString();
+
+		bool bAddNewGroup = true;
+		for (const TSharedPtr<struct FPropertySortedParamData> GroupChild : SortedParameters)
+		{
+			if (GroupChild->NodeKey == GroupProperty->NodeKey)
+			{
+				bAddNewGroup = false;
+			}
+		}
+		if (bAddNewGroup)
+		{
+			SortedParameters.Add(GroupProperty);
+		}
+
+		TSharedPtr<FPropertySortedParamData> ChildProperty(new FPropertySortedParamData());
+		ChildProperty->ParameterInfo = Property.ParameterInfo;
+		ChildProperty->StackDataType = EStackDataType::Property;
+		ChildProperty->ParameterNode = Property.ParameterNode;
+		ChildProperty->NodeKey = FString::FromInt(ChildProperty->ParameterInfo.Index) + FString::FromInt(ChildProperty->ParameterInfo.Association) + Property.GroupName.ToString() + Property.ParameterInfo.Name.ToString();
+
+		TArray<TSharedRef<IDetailTreeNode>> ParamChildren;
+		Property.ParameterNode->GetChildren(ParamChildren);
+		for (int32 ParamChildIdx = 0; ParamChildIdx < ParamChildren.Num(); ParamChildIdx++)
+		{
+			TSharedPtr<FPropertySortedParamData> ParamChildProperty(new FPropertySortedParamData());
+			ParamChildProperty->StackDataType = EStackDataType::PropertyChild;
+			ParamChildProperty->ParameterNode = ParamChildren[ParamChildIdx];
+			ParamChildProperty->ParameterHandle = ParamChildProperty->ParameterNode->CreatePropertyHandle();
+			ParamChildProperty->ParameterInfo = Property.ParameterInfo;
+			ChildProperty->Children.Add(ParamChildProperty);
+		}
+
+		for (TSharedPtr<struct FPropertySortedParamData> GroupChild : SortedParameters)
+		{
+			if (GroupChild->GroupName == Property.GroupName
+				&& GroupChild->ParameterInfo.Association == ChildProperty->ParameterInfo.Association
+				&&  GroupChild->ParameterInfo.Index == ChildProperty->ParameterInfo.Index)
+			{
+				GroupChild->Children.Add(ChildProperty);
+			}
+		}
+	}
+}
+
+void SMaterialDynamicParametersPanelWidget::Refresh()
+{
+	TSharedPtr<SHorizontalBox> HeaderBox;
+	NestedTree->CreateGroupsWidget();
+
+	if (NestedTree->HasAnyParameters())
+	{
+		this->ChildSlot
+		[
+			SNew(SVerticalBox)
+			+ SVerticalBox::Slot()
+			.Padding(0.0f)
+			[
+				
+				SNew(SVerticalBox)
+				+ SVerticalBox::Slot()
+				.AutoHeight()
+				[
+					SAssignNew(HeaderBox, SHorizontalBox)
+				]
+				+ SVerticalBox::Slot()
+				.Padding(0.0f)
+				[
+					SNew(SHorizontalBox)
+					+ SHorizontalBox::Slot()
+					.HAlign(HAlign_Fill)
+					[
+						NestedTree.ToSharedRef()
+					]
+					+ SHorizontalBox::Slot()
+					.HAlign(HAlign_Right)
+					.AutoWidth()
+					[
+						SNew(SBox)
+						.WidthOverride(16.0f)
+						[
+							ExternalScrollbar.ToSharedRef()
+						]
+					]
+					
+				]
+				
+			]
+		];
+
+		HeaderBox->AddSlot()
+			.FillWidth(1.0f)
+			[
+				SNullWidget::NullWidget
+			];
+	}
+	else
+	{
+		this->ChildSlot
+			[
+				SNew(SBox)
+				.Padding(10.0f)
+				.HAlign(HAlign_Center)
+				[
+					SNew(STextBlock)
+					.Text(LOCTEXT("ConnectMaterialParametersToFillList", "Connect a parameter to see it here."))
+				]
+			];
+	}
+}
+
+void SMaterialDynamicParametersPanelWidget::Construct(const FArguments& InArgs)
+{
+	MaterialInstance = InArgs._InMaterialInstance;
+	ExternalScrollbar = SNew(SScrollBar);
+
+	FPropertyEditorModule& PropertyEditorModule = FModuleManager::LoadModuleChecked<FPropertyEditorModule>("PropertyEditor");
+	const FPropertyRowGeneratorArgs PropertyRowGeneratorArgs;
+	Generator = PropertyEditorModule.CreatePropertyRowGenerator(PropertyRowGeneratorArgs);
+
+	Generator->SetObjects({ MaterialInstance.Get() });
+
+	NestedTree = SNew(SMaterialDynamicParametersOverviewTree)
+		.InMaterialInstance(InArgs._InMaterialInstance)
+		.InOwner(SharedThis(this))
+		.InScrollbar(ExternalScrollbar);
+
+	
+	Refresh();
+}
+
+void SMaterialDynamicParametersPanelWidget::UpdateInstance(UMaterialInstance* InMaterialInstance)
+{
+	Generator->SetObjects({ InMaterialInstance });
+	NestedTree->MaterialInstance = InMaterialInstance;
+	Refresh();
+}
+
+#undef LOCTEXT_NAMESPACE
diff --git a/Engine/Plugins/Experimental/VirtualProductionUtilities/Source/VPMaterialsEditor/Private/SMaterialDynamicParametersPanelWidget.h b/Engine/Plugins/Experimental/VirtualProductionUtilities/Source/VPMaterialsEditor/Private/SMaterialDynamicParametersPanelWidget.h
new file mode 100644
--- /dev/null
+++ b/Engine/Plugins/Experimental/VirtualProductionUtilities/Source/VPMaterialsEditor/Private/SMaterialDynamicParametersPanelWidget.h
@@ -0,0 +1,51 @@
+// Copyright Epic Games, Inc. All Rights Reserved.
+
+#pragma once
+
+#include "CoreMinimal.h"
+#include "Widgets/DeclarativeSyntaxSupport.h"
+#include "Widgets/SCompoundWidget.h"
+
+class IPropertyRowGenerator;
+class SMaterialDynamicParametersOverviewTree;
+class UMaterialInstance;
+
+/**
+ * Dynamic Material parameters view. The view creates all widgets for dynamic material instance parameters
+ * Most of the logic copied from SMaterialParametersOverviewTreeItem
+ */
+class SMaterialDynamicParametersPanelWidget : public SCompoundWidget
+{
+	friend SMaterialDynamicParametersOverviewTree;
+	
+public:
+	SLATE_BEGIN_ARGS(SMaterialDynamicParametersPanelWidget)
+		: _InMaterialInstance(nullptr)
+	{}
+
+	SLATE_ARGUMENT(UMaterialInstance*, InMaterialInstance)
+	SLATE_END_ARGS()
+
+	void Construct(const FArguments& InArgs);
+
+	/** Update material instance for tree generation and refresh */
+	void UpdateInstance(UMaterialInstance* InMaterialInstance);
+
+private:
+	/** Refresh MaterialParameters tree */
+	void Refresh();
+
+private:
+	/** The set of material parameters this is associated with */
+	TWeakObjectPtr<UMaterialInstance> MaterialInstance;
+
+	/** The tree contained in this item */
+	TSharedPtr<class SMaterialDynamicParametersOverviewTree> NestedTree;
+
+	/** Optional external Scroll bar */
+	TSharedPtr<class SScrollBar> ExternalScrollbar;
+
+	/** Row generator instance */
+	TSharedPtr<IPropertyRowGenerator> Generator;
+};
+
diff --git a/Engine/Plugins/Experimental/VirtualProductionUtilities/Source/VPMaterialsEditor/Private/SMaterialDynamicWidgets.cpp b/Engine/Plugins/Experimental/VirtualProductionUtilities/Source/VPMaterialsEditor/Private/SMaterialDynamicWidgets.cpp
new file mode 100644
--- /dev/null
+++ b/Engine/Plugins/Experimental/VirtualProductionUtilities/Source/VPMaterialsEditor/Private/SMaterialDynamicWidgets.cpp
@@ -0,0 +1,203 @@
+﻿// Copyright Epic Games, Inc. All Rights Reserved.
+
+#include "SMaterialDynamicWidgets.h"
+
+#include "PropertyCustomizationHelpers.h"
+#include "SMaterialDynamicParametersPanelWidget.h"
+#include "Components/ActorComponent.h"
+#include "Materials/MaterialInstanceDynamic.h"
+#include "Widgets/Input/SButton.h"
+#include "Widgets/Layout/SBox.h"
+#include "Widgets/Text/STextBlock.h"
+
+#define LOCTEXT_NAMESPACE "FMaterialDynamicList"
+
+void SMaterialDynamicView::Construct(const FArguments& InArgs, const TSharedRef<FMaterialItemView>& InMaterialItemView, UActorComponent* InCurrentComponent)
+{
+	MaterialItemViewWeakPtr = InMaterialItemView;
+	CurrentComponent = InCurrentComponent;
+
+	TSharedRef<SVerticalBox> ResultWidget = SNew(SVerticalBox);
+	if (InCurrentComponent != nullptr)
+	{
+		UMaterialInterface* MaterialInterface = InMaterialItemView->GetMaterialListItem().Material.Get();
+		UMaterialInstance* MaterialInstance = Cast<UMaterialInstance>(MaterialInterface);
+		UMaterialInstanceDynamic* MaterialInstanceDynamic = Cast<UMaterialInstanceDynamic>(MaterialInterface);
+		
+		if (MaterialInstance)
+		{
+			const bool bIsInvertedVisibility = !!MaterialInstanceDynamic;
+			const TAttribute<EVisibility> CreateDynamicMaterialButtonVisibility = GetButtonVisibilityAttribute<UMaterialInstance>(bIsInvertedVisibility);
+			
+			ResultWidget->AddSlot()
+			.AutoHeight()
+			[
+				SNew(SBox)
+				[
+					SNew(SHorizontalBox)
+					+SHorizontalBox::Slot()
+					.AutoWidth()
+					[
+						SNew(SButton)
+						.Visibility(GetButtonVisibilityAttribute<UMaterialInstanceDynamic>())
+						.OnClicked(this, &SMaterialDynamicView::OnRevertButtonClicked)
+						[
+							SNew(STextBlock).Text(LOCTEXT("Revert_Button", "Revert"))
+						]
+					]
+
+					+SHorizontalBox::Slot()
+					.AutoWidth()
+					[
+						SNew(SButton)
+						.Visibility(GetButtonVisibilityAttribute<UMaterialInstanceDynamic>())
+						.OnClicked(this, &SMaterialDynamicView::OnResetButtonClicked)
+						[
+							SNew(STextBlock).Text(LOCTEXT("Reset_Button", "Reset"))
+						]
+					]
+
+					+SHorizontalBox::Slot()
+					.AutoWidth()
+					[
+						SNew(SButton)
+						.Visibility(GetButtonVisibilityAttribute<UMaterialInstanceDynamic>())
+						.OnClicked(this, &SMaterialDynamicView::OnCopyToOriginalButtonClicked)
+						[
+							SNew(STextBlock).Text(LOCTEXT("CopyToOriginal_Button", "Copy to Original"))
+						]
+					]
+
+					+SHorizontalBox::Slot()
+					.AutoWidth()
+					[
+						SNew(SButton)
+						.Visibility(CreateDynamicMaterialButtonVisibility)
+						.OnClicked(this, &SMaterialDynamicView::OnCreateDynamicMaterialButtonClicked)
+						[
+							SNew(STextBlock).Text(LOCTEXT("CreateDynamicMaterial_Button", "Create Dynamic Material"))
+						]
+					]
+				]
+			];
+		}
+
+		if (MaterialInstanceDynamic)
+		{
+			ResultWidget->AddSlot()
+			.AutoHeight()
+			[
+				SNew(SMaterialDynamicParametersPanelWidget)
+				.InMaterialInstance(MaterialInstanceDynamic)
+			];
+		}
+	}
+
+	ChildSlot
+	[
+		ResultWidget
+	];
+}
+
+FReply SMaterialDynamicView::OnResetButtonClicked() const
+{
+	const TSharedPtr<FMaterialItemView> MaterialItemView = MaterialItemViewWeakPtr.Pin();
+	if(!ensure(MaterialItemView.IsValid()))
+	{
+		return FReply::Handled();
+	}
+	
+	UMaterialInstanceDynamic* MaterialInstanceDynamic = Cast<UMaterialInstanceDynamic>(MaterialItemView->GetMaterialListItem().Material.Get());
+	if (!MaterialInstanceDynamic)
+	{
+		return FReply::Handled();
+	}
+
+	UMaterialInstance* ParentMaterialInstance = Cast<UMaterialInstance>(MaterialInstanceDynamic->Parent);
+	if (!ParentMaterialInstance)
+	{
+		return FReply::Handled();
+	}
+
+	MaterialInstanceDynamic->CopyParameterOverrides(ParentMaterialInstance);
+
+	return FReply::Handled();
+}
+
+FReply SMaterialDynamicView::OnRevertButtonClicked() const
+{
+	const TSharedPtr<FMaterialItemView> MaterialItemView = MaterialItemViewWeakPtr.Pin();
+	if(!ensure(MaterialItemView.IsValid()))
+	{
+		return FReply::Handled();
+	}
+	
+	UMaterialInstance* ParentMaterialInstance = GetMaterialParent<UMaterialInstance, UMaterialInstanceDynamic>(MaterialItemView->GetMaterialListItem().Material.Get());
+	if (!ParentMaterialInstance)
+	{
+		return FReply::Handled();
+	}
+
+	MaterialItemView->ReplaceMaterial(ParentMaterialInstance);
+
+	return FReply::Handled();
+}
+
+FReply SMaterialDynamicView::OnCopyToOriginalButtonClicked() const
+{
+	const TSharedPtr<FMaterialItemView> MaterialItemView = MaterialItemViewWeakPtr.Pin();
+	if(!ensure(MaterialItemView.IsValid()))
+	{
+		return FReply::Handled();
+	}
+	
+	UMaterialInstanceDynamic* MaterialInstanceDynamic = Cast<UMaterialInstanceDynamic>(MaterialItemView->GetMaterialListItem().Material.Get());
+	if (!MaterialInstanceDynamic)
+	{
+		return FReply::Handled();
+	}
+
+	UMaterialInstance* ParentMaterialInstance = Cast<UMaterialInstance>(MaterialInstanceDynamic->Parent);
+	if (!ParentMaterialInstance)
+	{
+		return FReply::Handled();
+	}
+
+	ParentMaterialInstance->ScalarParameterValues = MaterialInstanceDynamic->ScalarParameterValues;
+	ParentMaterialInstance->VectorParameterValues = MaterialInstanceDynamic->VectorParameterValues;
+	ParentMaterialInstance->TextureParameterValues = MaterialInstanceDynamic->TextureParameterValues;
+	ParentMaterialInstance->RuntimeVirtualTextureParameterValues = MaterialInstanceDynamic->RuntimeVirtualTextureParameterValues;
+	ParentMaterialInstance->FontParameterValues = MaterialInstanceDynamic->FontParameterValues;
+	ParentMaterialInstance->UpdateStaticPermutation();
+	ParentMaterialInstance->Modify();
+
+	return FReply::Handled();
+}
+
+FReply SMaterialDynamicView::OnCreateDynamicMaterialButtonClicked() const
+{
+	const TSharedPtr<FMaterialItemView> MaterialItemView = MaterialItemViewWeakPtr.Pin();
+	if(!ensure(MaterialItemView.IsValid()))
+	{
+		return FReply::Handled();
+	}
+	
+	UMaterialInstance* MaterialInstance = Cast<UMaterialInstance>(MaterialItemView->GetMaterialListItem().Material.Get());
+
+	if (!CurrentComponent.IsValid())
+	{
+		return FReply::Handled();
+	}
+
+	if (MaterialInstance)
+	{
+		UMaterialInstanceDynamic* NewMaterialInstanceDynamic = UMaterialInstanceDynamic::Create(MaterialInstance, CurrentComponent.Get());
+		NewMaterialInstanceDynamic->CopyParameterOverrides(MaterialInstance);
+		MaterialItemView->ReplaceMaterial(NewMaterialInstanceDynamic);
+	}
+
+	return FReply::Handled();
+}
+
+#undef LOCTEXT_NAMESPACE
+
diff --git a/Engine/Plugins/Experimental/VirtualProductionUtilities/Source/VPMaterialsEditor/Private/SMaterialDynamicWidgets.h b/Engine/Plugins/Experimental/VirtualProductionUtilities/Source/VPMaterialsEditor/Private/SMaterialDynamicWidgets.h
new file mode 100644
--- /dev/null
+++ b/Engine/Plugins/Experimental/VirtualProductionUtilities/Source/VPMaterialsEditor/Private/SMaterialDynamicWidgets.h
@@ -0,0 +1,78 @@
+﻿// Copyright Epic Games, Inc. All Rights Reserved.
+
+#pragma once
+
+#include "CoreMinimal.h"
+#include "MaterialList.h"
+#include "Widgets/DeclarativeSyntaxSupport.h"
+#include "Widgets/SCompoundWidget.h"
+
+class IPropertyHandle;
+class IDetailLayoutBuilder;
+class UActorComponent;
+class UMaterialInterface;
+
+/**
+ * Creates the view for material dynamic view
+ */
+class SMaterialDynamicView : public SCompoundWidget
+{
+public:
+	SLATE_BEGIN_ARGS(SMaterialDynamicView)
+	{}
+	SLATE_END_ARGS()
+
+	void Construct(const FArguments& InArgs, const TSharedRef<FMaterialItemView>& InMaterialItemView, UActorComponent* InCurrentComponent);
+
+private:
+	/** Get Visibility for extension button */
+	template<typename TMaterialClass>
+	TAttribute<EVisibility> GetButtonVisibilityAttribute(const bool bIsInvertedResult = false) const
+	{
+		return MakeAttributeLambda([this, bIsInvertedResult]()
+			{
+				const TSharedPtr<FMaterialItemView> MaterialItemView = MaterialItemViewWeakPtr.Pin();
+				if(!ensure(MaterialItemView.IsValid()))
+				{
+					return EVisibility::Collapsed;
+				}
+			
+				bool IsVisible = Cast<TMaterialClass>(MaterialItemView->GetMaterialListItem().Material.Get()) ? true : false;
+				IsVisible ^= bIsInvertedResult;
+
+				return IsVisible ? EVisibility::Visible : EVisibility::Collapsed;
+			});
+	};
+
+	/** Button handler for revert to original material for the slot */
+	FReply OnRevertButtonClicked() const;
+
+	/** Button handler for reset parameters to default instance material parameter */
+	FReply OnResetButtonClicked() const;
+
+	/** Button handler for copy material parameters to original instance material */
+	FReply OnCopyToOriginalButtonClicked() const;
+
+	/** Button handler for create dynamic material */
+	FReply OnCreateDynamicMaterialButtonClicked() const;
+
+	/** Get parent material and cast to specified type */
+	template<typename TReturnType, typename TInputType>
+	TReturnType* GetMaterialParent(UMaterialInterface* InMaterialInterface) const
+	{
+		if (TInputType* Material = Cast<TInputType>(InMaterialInterface))
+		{
+			return Cast<TReturnType>(Material->Parent);
+		}
+
+		return nullptr;
+	}
+
+private:
+	/** Weak Pointer to Material Item View */
+	TWeakPtr<FMaterialItemView> MaterialItemViewWeakPtr;
+
+	/** Weak pointer to current component using the current material */
+	TWeakObjectPtr<UActorComponent> CurrentComponent;
+};
+
diff --git a/Engine/Plugins/Experimental/VirtualProductionUtilities/Source/VPMaterialsEditor/Private/VPMaterialsEditorModule.cpp b/Engine/Plugins/Experimental/VirtualProductionUtilities/Source/VPMaterialsEditor/Private/VPMaterialsEditorModule.cpp
new file mode 100644
--- /dev/null
+++ b/Engine/Plugins/Experimental/VirtualProductionUtilities/Source/VPMaterialsEditor/Private/VPMaterialsEditorModule.cpp
@@ -0,0 +1,24 @@
+// Copyright Epic Games, Inc. All Rights Reserved.
+
+#include "VPMaterialsEditorModule.h"
+#include "SMaterialDynamicWidgets.h"
+
+#include "MaterialList.h"
+
+void FVPMaterialsEditorModule::StartupModule()
+{
+	// Add bottom extender for material item
+	FMaterialList::OnAddMaterialItemViewExtraBottomWidget.AddLambda([](const TSharedRef<FMaterialItemView>& InMaterialItemView, UActorComponent* InCurrentComponent, IDetailLayoutBuilder& InDetailBuilder, TArray<TSharedPtr<SWidget>>& OutExtensions)
+	{
+		OutExtensions.Add(SNew(SMaterialDynamicView, InMaterialItemView, InCurrentComponent));
+	});
+}
+
+void FVPMaterialsEditorModule::ShutdownModule()
+{
+	FMaterialList::OnAddMaterialItemViewExtraBottomWidget.RemoveAll(this);
+}
+
+IMPLEMENT_MODULE(FVPMaterialsEditorModule, VPMaterialsEditor)
+
+			
\ No newline at end of file
diff --git a/Engine/Plugins/Experimental/VirtualProductionUtilities/Source/VPMaterialsEditor/Private/VPMaterialsEditorModule.h b/Engine/Plugins/Experimental/VirtualProductionUtilities/Source/VPMaterialsEditor/Private/VPMaterialsEditorModule.h
new file mode 100644
--- /dev/null
+++ b/Engine/Plugins/Experimental/VirtualProductionUtilities/Source/VPMaterialsEditor/Private/VPMaterialsEditorModule.h
@@ -0,0 +1,19 @@
+// Copyright Epic Games, Inc. All Rights Reserved.
+
+#pragma once
+
+#include "CoreMinimal.h"
+#include "Modules/ModuleInterface.h"
+#include "Modules/ModuleManager.h"
+
+/**
+ * Utility module for extend material editors and details panel
+ */
+class FVPMaterialsEditorModule : public IModuleInterface
+{
+public:
+	//~ Begin IModuleInterface interface
+	virtual void StartupModule() override;
+	virtual void ShutdownModule() override;
+	//~ End IModuleInterface interface
+};
diff --git a/Engine/Plugins/Experimental/VirtualProductionUtilities/Source/VPMaterialsEditor/VPMaterialsEditor.Build.cs b/Engine/Plugins/Experimental/VirtualProductionUtilities/Source/VPMaterialsEditor/VPMaterialsEditor.Build.cs
new file mode 100644
--- /dev/null
+++ b/Engine/Plugins/Experimental/VirtualProductionUtilities/Source/VPMaterialsEditor/VPMaterialsEditor.Build.cs
@@ -0,0 +1,27 @@
+// Copyright Epic Games, Inc. All Rights Reserved.
+
+using UnrealBuildTool;
+
+public class VPMaterialsEditor : ModuleRules
+{
+	public VPMaterialsEditor(ReadOnlyTargetRules Target) : base(Target)
+	{
+		PCHUsage = ModuleRules.PCHUsageMode.UseExplicitOrSharedPCHs;
+
+		PrivateDependencyModuleNames.AddRange(
+			new string[]
+			{
+				"Core",
+				"CoreUObject",
+				"EditorStyle",
+				"Engine",
+				"InputCore",
+				"Landscape",
+				"PropertyEditor",
+				"Slate",
+				"SlateCore",
+				"UnrealEd",
+            }
+		);
+	}
+}
diff --git a/Engine/Plugins/Experimental/VirtualProductionUtilities/VirtualProductionUtilities.uplugin b/Engine/Plugins/Experimental/VirtualProductionUtilities/VirtualProductionUtilities.uplugin
--- a/Engine/Plugins/Experimental/VirtualProductionUtilities/VirtualProductionUtilities.uplugin
+++ b/Engine/Plugins/Experimental/VirtualProductionUtilities/VirtualProductionUtilities.uplugin
@@ -33,6 +33,11 @@
 			"Name": "VPBookmarkEditor",
 			"Type": "Editor",
 			"LoadingPhase": "PostDefault"
+		},
+		{
+			"Name": "VPMaterialsEditor",
+			"Type": "Editor",
+			"LoadingPhase": "PostDefault"
 		}
 	],
 	"Plugins": [
diff --git a/Engine/Source/Editor/DetailCustomizations/Private/ComponentMaterialCategory.cpp b/Engine/Source/Editor/DetailCustomizations/Private/ComponentMaterialCategory.cpp
--- a/Engine/Source/Editor/DetailCustomizations/Private/ComponentMaterialCategory.cpp
+++ b/Engine/Source/Editor/DetailCustomizations/Private/ComponentMaterialCategory.cpp
@@ -220,7 +220,7 @@ void FComponentMaterialCategory::OnGetMaterialsForView( IMaterialListBuilder& Ma
 			// Add the material if we allow null materials to be added or we have a valid material
 			if( bAllowNullEntries || Material )
 			{
-				MaterialList.AddMaterial( MaterialIndex, Material, bCanBeReplaced );
+				MaterialList.AddMaterial( MaterialIndex, Material, bCanBeReplaced, CurrentComponent );
 				bAnyMaterialsToDisplay = true;
 			}
 		}
diff --git a/Engine/Source/Editor/PropertyEditor/Private/MaterialList.cpp b/Engine/Source/Editor/PropertyEditor/Private/MaterialList.cpp
--- a/Engine/Source/Editor/PropertyEditor/Private/MaterialList.cpp
+++ b/Engine/Source/Editor/PropertyEditor/Private/MaterialList.cpp
@@ -5,17 +5,22 @@
 #include "Editor.h"
 #include "IDetailChildrenBuilder.h"
 #include "PropertyCustomizationHelpers.h"
+#include "PropertyEditorModule.h"
 #include "PropertyHandle.h"
 #include "Framework/MultiBox/MultiBoxBuilder.h"
+#include "Modules/ModuleManager.h"
+#include "Styling/SlateIconFinder.h"
+#include "Widgets/Images/SImage.h"
 #include "Widgets/Input/SComboButton.h"
 #include "Widgets/Input/SHyperlink.h"
 #include "Widgets/Layout/SBox.h"
 #include "Widgets/Text/STextBlock.h"
-#include "Widgets/Images/SImage.h"
-#include "Styling/SlateIconFinder.h"
 
 #define LOCTEXT_NAMESPACE "PropertyCustomizationHelpers"
 
+/** Definition of extra widgets delegate for bottom material value field. */
+FMaterialList::FOnAddMaterialItemViewExtraBottomWidget FMaterialList::OnAddMaterialItemViewExtraBottomWidget;
+
 /**
  * Builds up a list of unique materials while creating some information about the materials
  */
@@ -31,11 +36,11 @@ public:
 	 * @param Material		The material being used
 	 * @param bCanBeReplced	Whether or not the material can be replaced by a user
 	 */
-	virtual void AddMaterial( uint32 SlotIndex, UMaterialInterface* Material, bool bCanBeReplaced ) override
+	virtual void AddMaterial( uint32 SlotIndex, UMaterialInterface* Material, bool bCanBeReplaced, UActorComponent* InCurrentComponent ) override
 	{
 		int32 NumMaterials = MaterialSlots.Num();
 
-		FMaterialListItem MaterialItem( Material, SlotIndex, bCanBeReplaced ); 
+		FMaterialListItem MaterialItem( Material, SlotIndex, bCanBeReplaced, InCurrentComponent ); 
 		if( !UniqueMaterials.Contains( MaterialItem ) ) 
 		{
 			MaterialSlots.Add( MaterialItem );
@@ -95,69 +100,91 @@ private:
 	TArray<uint32> MaterialCount;
 };
 
-/**
- * A view of a single item in an FMaterialList
- */
-class FMaterialItemView : public TSharedFromThis<FMaterialItemView>
+
+TSharedRef<FMaterialItemView> FMaterialItemView::Create(
+	const FMaterialListItem& Material, 
+	FOnMaterialChanged InOnMaterialChanged,
+	FOnGenerateWidgetsForMaterial InOnGenerateNameWidgetsForMaterial, 
+	FOnGenerateWidgetsForMaterial InOnGenerateWidgetsForMaterial, 
+	FOnResetMaterialToDefaultClicked InOnResetToDefaultClicked,
+	int32 InMultipleMaterialCount,
+	bool bShowUsedTextures)
 {
-public:
-	/**
-	 * Creates a new instance of this class
-	 *
-	 * @param Material				The material to view
-	 * @param InOnMaterialChanged	Delegate for when the material changes
-	 */
-	static TSharedRef<FMaterialItemView> Create(
-		const FMaterialListItem& Material, 
-		FOnMaterialChanged InOnMaterialChanged,
-		FOnGenerateWidgetsForMaterial InOnGenerateNameWidgetsForMaterial, 
-		FOnGenerateWidgetsForMaterial InOnGenerateWidgetsForMaterial, 
-		FOnResetMaterialToDefaultClicked InOnResetToDefaultClicked,
-		int32 InMultipleMaterialCount,
-		bool bShowUsedTextures)
-	{
-		return MakeShareable( new FMaterialItemView( Material, InOnMaterialChanged, InOnGenerateNameWidgetsForMaterial, InOnGenerateWidgetsForMaterial, InOnResetToDefaultClicked, InMultipleMaterialCount, bShowUsedTextures) );
-	}
+	// FMaterialItemView has private constructor that is why we need to use MakeShareable 
+	return MakeShareable( new FMaterialItemView( Material, InOnMaterialChanged, InOnGenerateNameWidgetsForMaterial, InOnGenerateWidgetsForMaterial, InOnResetToDefaultClicked, InMultipleMaterialCount, bShowUsedTextures) );
+}
 
-	TSharedRef<SWidget> CreateNameContent()
+TSharedRef<SWidget> FMaterialItemView::CreateNameContent()
+{
+	FFormatNamedArguments Arguments;
+	Arguments.Add(TEXT("ElementIndex"), MaterialItem.SlotIndex);
+
+	return 
+		SNew(SVerticalBox)
+		+SVerticalBox::Slot()
+		.VAlign(VAlign_Center)
+		[
+			SNew( STextBlock )
+			.Font( IDetailLayoutBuilder::GetDetailFont() )
+			.Text( FText::Format(LOCTEXT("ElementIndex", "Element {ElementIndex}"), Arguments ) )
+			.OverflowPolicy(ETextOverflowPolicy::Ellipsis)
+		]
+		+SVerticalBox::Slot()
+		.Padding(0.0f,4.0f)
+		.AutoHeight()
+		[
+			OnGenerateCustomNameWidgets.IsBound() ? OnGenerateCustomNameWidgets.Execute( MaterialItem.Material.Get(), MaterialItem.SlotIndex ) : StaticCastSharedRef<SWidget>( SNullWidget::NullWidget )
+		];
+}
+
+FString FMaterialItemView::OnGetObjectPath() const
+{
+	return MaterialItem.Material->GetPathName();
+}
+
+namespace ValueExtender
+{
+	template <typename TContainer, typename TDelegate>
+	TSharedRef<SWidget> MakeWidget(TDelegate& InDelegate, const TSharedRef<FMaterialItemView>& InMaterialItemView, IDetailLayoutBuilder& InDetailBuilder, UActorComponent* InCurrentComponent)
 	{
-		FFormatNamedArguments Arguments;
-		Arguments.Add(TEXT("ElementIndex"), MaterialItem.SlotIndex);
+		TSharedPtr<TContainer> WidgetContainer = SNew(TContainer);
+		TArray<TSharedPtr<SWidget>> Widgets;
 
-		return 
-			SNew(SVerticalBox)
-			+SVerticalBox::Slot()
-			.VAlign(VAlign_Center)
-			[
-				SNew( STextBlock )
-				.Font( IDetailLayoutBuilder::GetDetailFont() )
-				.Text( FText::Format(LOCTEXT("ElementIndex", "Element {ElementIndex}"), Arguments ) )
-				.OverflowPolicy(ETextOverflowPolicy::Ellipsis)
-			]
-			+SVerticalBox::Slot()
-			.Padding(0.0f,4.0f)
-			.AutoHeight()
+		// Execute all delegates and loop through all results widgets
+		InDelegate.Broadcast(InMaterialItemView, InCurrentComponent, InDetailBuilder, Widgets);
+
+		for (const TSharedPtr<SWidget> Widget : Widgets)
+		{
+			WidgetContainer->AddSlot()
 			[
-				OnGenerateCustomNameWidgets.IsBound() ? OnGenerateCustomNameWidgets.Execute( MaterialItem.Material.Get(), MaterialItem.SlotIndex ) : StaticCastSharedRef<SWidget>( SNullWidget::NullWidget )
+				Widget.ToSharedRef()
 			];
+		}
+
+		return WidgetContainer.ToSharedRef();
 	}
+};
 
-	TSharedRef<SWidget> CreateValueContent( const TSharedPtr<FAssetThumbnailPool>& ThumbnailPool, const TArray<FAssetData>& OwnerAssetDataArray)
-	{
-		return
-			SNew(SVerticalBox)
-			+SVerticalBox::Slot()
-			.AutoHeight()
-			.Padding( 0.0f )
-			.VAlign(VAlign_Center)
-			.HAlign(HAlign_Fill)
+TSharedRef<SWidget> FMaterialItemView::CreateValueContent(IDetailLayoutBuilder& InDetailBuilder, const TArray<FAssetData>& OwnerAssetDataArray, UActorComponent* InActorComponent)
+{	
+	return
+		SNew(SVerticalBox)
+		+SVerticalBox::Slot()
+		.AutoHeight()
+		.Padding( 0.0f )
+		.VAlign(VAlign_Center)
+		.HAlign(HAlign_Fill)
+		[
+			SNew(SHorizontalBox)
+			+SHorizontalBox::Slot()
+			.AutoWidth()
 			[
 				SNew(SObjectPropertyEntryBox)
 				.ObjectPath(this, &FMaterialItemView::OnGetObjectPath)
 				.AllowClear(false)
 				.AllowedClass(UMaterialInterface::StaticClass())
 				.OnObjectChanged(this, &FMaterialItemView::OnSetObject)
-				.ThumbnailPool(ThumbnailPool)
+				.ThumbnailPool(InDetailBuilder.GetThumbnailPool())
 				.DisplayCompactSize(true)
 				.OwnerAssetDataArray(OwnerAssetDataArray)
 				.CustomContentSlot()
@@ -188,142 +215,251 @@ public:
 						]
 						+SHorizontalBox::Slot()
 						.Padding(3.0f, 0.0f)
-						.FillWidth(1.0f)
+						.AutoWidth()
 						[
 							OnGenerateCustomMaterialWidgets.IsBound() ? OnGenerateCustomMaterialWidgets.Execute(MaterialItem.Material.Get(), MaterialItem.SlotIndex) : SNullWidget::NullWidget
 						]
 					]
 				]
-			];
+			]
+			+SHorizontalBox::Slot()
+			.AutoWidth()
+			[
+				GetGlobalRowExtensionWidget(InDetailBuilder, InActorComponent).ToSharedRef()
+			]
+		]
+		+SVerticalBox::Slot()
+		.AutoHeight()
+		[
+			ValueExtender::MakeWidget<SVerticalBox>(FMaterialList::OnAddMaterialItemViewExtraBottomWidget, SharedThis(this), InDetailBuilder, InActorComponent)
+		];
+}
+
+bool FMaterialItemView::GetResetToBaseVisibility() const
+{
+	// Only show the reset to base button if the current material can be replaced
+	return OnMaterialChanged.IsBound() && MaterialItem.bCanBeReplaced;
+}
+
+
+void FMaterialItemView::OnResetToBaseClicked()
+{
+	// Only allow reset to base if the current material can be replaced
+	if( MaterialItem.Material.IsValid() && MaterialItem.bCanBeReplaced )
+	{
+		bool bReplaceAll = false;
+		ReplaceMaterial( nullptr, bReplaceAll );
+		OnResetToDefaultClicked.ExecuteIfBound( MaterialItem.Material.Get(), MaterialItem.SlotIndex );
 	}
+}
 
-	/**
-	 * Called to get the visibility of the reset to base button
-	 */
-	bool GetResetToBaseVisibility() const
+FMaterialItemView::FMaterialItemView(	const FMaterialListItem& InMaterial, 
+					FOnMaterialChanged& InOnMaterialChanged, 
+					FOnGenerateWidgetsForMaterial& InOnGenerateNameWidgets, 
+					FOnGenerateWidgetsForMaterial& InOnGenerateMaterialWidgets, 
+					FOnResetMaterialToDefaultClicked& InOnResetToDefaultClicked,
+					int32 InMultipleMaterialCount,
+					bool bInShowUsedTextures)
+					
+	: MaterialItem( InMaterial )
+	, OnMaterialChanged( InOnMaterialChanged )
+	, OnGenerateCustomNameWidgets( InOnGenerateNameWidgets )
+	, OnGenerateCustomMaterialWidgets( InOnGenerateMaterialWidgets )
+	, OnResetToDefaultClicked( InOnResetToDefaultClicked )
+	, MultipleMaterialCount( InMultipleMaterialCount )
+	, bShowUsedTextures( bInShowUsedTextures )
+{
+
+}
+
+void FMaterialItemView::ReplaceMaterial( UMaterialInterface* NewMaterial, bool bReplaceAll )
+{
+	UMaterialInterface* PrevMaterial = NULL;
+	if( MaterialItem.Material.IsValid() )
 	{
-		// Only show the reset to base button if the current material can be replaced
-		return OnMaterialChanged.IsBound() && MaterialItem.bCanBeReplaced;
+		PrevMaterial = MaterialItem.Material.Get();
 	}
 
-	/**
-	 * Called when reset to base is clicked
-	 */
-	void OnResetToBaseClicked()
+	if( NewMaterial != PrevMaterial )
 	{
-		// Only allow reset to base if the current material can be replaced
-		if( MaterialItem.Material.IsValid() && MaterialItem.bCanBeReplaced )
-		{
-			bool bReplaceAll = false;
-			ReplaceMaterial( nullptr, bReplaceAll );
-			OnResetToDefaultClicked.ExecuteIfBound( MaterialItem.Material.Get(), MaterialItem.SlotIndex );
-		}
+		// Replace the material
+		OnMaterialChanged.ExecuteIfBound( NewMaterial, PrevMaterial, MaterialItem.SlotIndex, bReplaceAll );
 	}
+}
 
-private:
+void FMaterialItemView::OnSetObject( const FAssetData& AssetData )
+{
+	const bool bReplaceAll = false;
 
-	FMaterialItemView(	const FMaterialListItem& InMaterial, 
-						FOnMaterialChanged& InOnMaterialChanged, 
-						FOnGenerateWidgetsForMaterial& InOnGenerateNameWidgets, 
-						FOnGenerateWidgetsForMaterial& InOnGenerateMaterialWidgets, 
-						FOnResetMaterialToDefaultClicked& InOnResetToDefaultClicked,
-						int32 InMultipleMaterialCount,
-						bool bInShowUsedTextures)
-						
-		: MaterialItem( InMaterial )
-		, OnMaterialChanged( InOnMaterialChanged )
-		, OnGenerateCustomNameWidgets( InOnGenerateNameWidgets )
-		, OnGenerateCustomMaterialWidgets( InOnGenerateMaterialWidgets )
-		, OnResetToDefaultClicked( InOnResetToDefaultClicked )
-		, MultipleMaterialCount( InMultipleMaterialCount )
-		, bShowUsedTextures( bInShowUsedTextures )
-	{
+	UMaterialInterface* NewMaterial = Cast<UMaterialInterface>(AssetData.GetAsset());
+	ReplaceMaterial( NewMaterial, bReplaceAll );
+}
 
-	}
+bool FMaterialItemView::IsTexturesMenuEnabled() const
+{
+	return MaterialItem.Material.Get() != NULL;
+}
+
+TSharedRef<SWidget> FMaterialItemView::OnGetTexturesMenuForMaterial()
+{
+	FMenuBuilder MenuBuilder( true, NULL );
 
-	void ReplaceMaterial( UMaterialInterface* NewMaterial, bool bReplaceAll = false )
+	if( MaterialItem.Material.IsValid() )
 	{
-		UMaterialInterface* PrevMaterial = NULL;
-		if( MaterialItem.Material.IsValid() )
-		{
-			PrevMaterial = MaterialItem.Material.Get();
-		}
+		UMaterialInterface* Material = MaterialItem.Material.Get();
+
+		TArray< UTexture* > Textures;
+		Material->GetUsedTextures(Textures, EMaterialQualityLevel::Num, false, ERHIFeatureLevel::Num, true);
 
-		if( NewMaterial != PrevMaterial )
+		// Add a menu item for each texture.  Clicking on the texture will display it in the content browser
+		// UObject for delegate compatibility
+		for( UObject* Texture : Textures )
 		{
-			// Replace the material
-			OnMaterialChanged.ExecuteIfBound( NewMaterial, PrevMaterial, MaterialItem.SlotIndex, bReplaceAll );
+			FUIAction Action( FExecuteAction::CreateSP( this, &FMaterialItemView::GoToAssetInContentBrowser, MakeWeakObjectPtr(Texture) ) );
+
+			MenuBuilder.AddMenuEntry( FText::FromString( Texture->GetName() ), LOCTEXT( "BrowseTexture_ToolTip", "Find this texture in the content browser" ), FSlateIcon(), Action );
 		}
 	}
 
-	void OnSetObject( const FAssetData& AssetData )
-	{
-		const bool bReplaceAll = false;
+	return MenuBuilder.MakeWidget();
+}
 
-		UMaterialInterface* NewMaterial = Cast<UMaterialInterface>(AssetData.GetAsset());
-		ReplaceMaterial( NewMaterial, bReplaceAll );
-	}
 
-	FString OnGetObjectPath() const
+void FMaterialItemView::GoToAssetInContentBrowser( TWeakObjectPtr<UObject> Object )
+{
+	if( Object.IsValid() )
 	{
-		return MaterialItem.Material->GetPathName();
+		TArray< UObject* > Objects;
+		Objects.Add( Object.Get() );
+		GEditor->SyncBrowserToObjects( Objects );
 	}
+}
 
-	/**
-	 * @return Whether or not the textures menu is enabled
-	 */
-	bool IsTexturesMenuEnabled() const
+PRAGMA_DISABLE_OPTIMIZATION
+TSharedPtr<IPropertyHandle> FMaterialItemView::GetMaterialHandle(IDetailLayoutBuilder& InDetailBuilder, UActorComponent* InCurrentComponent) const
+{
+	USceneComponent* SceneComponent = Cast<USceneComponent>(InCurrentComponent);
+	if (!SceneComponent)
 	{
-		return MaterialItem.Material.Get() != NULL;
+		return nullptr;
 	}
 
-	TSharedRef<SWidget> OnGetTexturesMenuForMaterial()
+	TSharedPtr<IPropertyHandle> PropertyHandle;
+	UObject* OwnerObject = nullptr;
+	FString PropertyPath;
+	if (SceneComponent->GetMaterialPropertyPath(MaterialItem.SlotIndex, OwnerObject, PropertyPath))
 	{
-		FMenuBuilder MenuBuilder( true, NULL );
+		TArray<FString> PathArray;
+		// Split full path to array
+		PropertyPath.ParseIntoArray(PathArray, TEXT("."));
 
-		if( MaterialItem.Material.IsValid() )
+		for (int32 PathArrayIndex = 0; PathArrayIndex < PathArray.Num(); ++PathArrayIndex)
 		{
-			UMaterialInterface* Material = MaterialItem.Material.Get();
+			const FString& PathArrayItem = PathArray[PathArrayIndex];
+
+			int32 StartPathArrayIndex = -1;
+			int32 EndPathArrayIndex = -1;
+			PathArrayItem.FindChar('[', StartPathArrayIndex);
+			PathArrayItem.FindChar(']', EndPathArrayIndex);
 
-			TArray< UTexture* > Textures;
-			Material->GetUsedTextures(Textures, EMaterialQualityLevel::Num, false, ERHIFeatureLevel::Num, true);
+			// Get a property name without array 
+			FString PropertyName = (StartPathArrayIndex != -1) ? PathArrayItem.Mid(0, StartPathArrayIndex) : PathArrayItem;
 
-			// Add a menu item for each texture.  Clicking on the texture will display it in the content browser
-			// UObject for delegate compatibility
-			for( UObject* Texture : Textures )
+			// create a property handle if that is first path element
+			if (PathArrayIndex == 0)
 			{
-				FUIAction Action( FExecuteAction::CreateSP( this, &FMaterialItemView::GoToAssetInContentBrowser, MakeWeakObjectPtr(Texture) ) );
+				UClass* OwnerObjectClass = OwnerObject->GetClass();
+
+				// Exception for UStaticMesh
+				// StaticMaterials property does not contain EditAnywhere flag which is required for PropertyHandle generation
+				if (OwnerObjectClass == UStaticMesh::StaticClass())
+				{					
+					FProperty* Property = UStaticMesh::StaticClass()->FindPropertyByName(UStaticMesh::GetStaticMaterialsName());
+					check(Property);
+
+					const EPropertyFlags OriginalPropertyFlags = Property->GetPropertyFlags();
+					Property->SetPropertyFlags(CPF_Edit);
+					PropertyHandle = InDetailBuilder.AddObjectPropertyData( { OwnerObject }, *PropertyName);
+					Property->SetPropertyFlags(OriginalPropertyFlags);
+				}
+				else
+				{
+					PropertyHandle = InDetailBuilder.AddObjectPropertyData( { OwnerObject }, *PropertyName);
+				}				
+			}
 
-				MenuBuilder.AddMenuEntry( FText::FromString( Texture->GetName() ), LOCTEXT( "BrowseTexture_ToolTip", "Find this texture in the content browser" ), FSlateIcon(), Action );
+			if (PropertyHandle.IsValid())
+			{
+				// Then check if property path contain array element like Materials[0]
+				if(const TSharedPtr<IPropertyHandleArray> HandleArray = PropertyHandle->AsArray())
+				{
+					if (StartPathArrayIndex != INDEX_NONE && EndPathArrayIndex != INDEX_NONE)
+					{
+						const FString ArrayIndexStr = PathArrayItem.Mid(StartPathArrayIndex + 1, EndPathArrayIndex - 1 - StartPathArrayIndex);
+
+						int32 ArrayIndex = INDEX_NONE;
+						LexFromString(ArrayIndex, *ArrayIndexStr);
+
+						uint32 HandleArrayElementsNum = 0;
+						HandleArray->GetNumElements(HandleArrayElementsNum);
+						if (ensure(ArrayIndex != INDEX_NONE) && ensure(ArrayIndex < static_cast<int32>(HandleArrayElementsNum)))
+						{
+							PropertyHandle = HandleArray->GetElement(ArrayIndex);
+						}
+					}
+				}
+				// Property Handle is inner struct element if we have more then 1 element in path array PathArray
+				// For example Materials[0].MaterialInterface 
+				if (PathArrayIndex != 0)
+				{
+					// Make sure the property name is valid
+					if (ensure(!PropertyName.IsEmpty()))
+					{
+						PropertyHandle = PropertyHandle->GetChildHandle(*PropertyName);
+					}
+					else
+					{
+						PropertyHandle = nullptr;
+					}
+				}
 			}
 		}
-
-		return MenuBuilder.MakeWidget();
 	}
 
-	/**
-	 * Finds the asset in the content browser
-	 */
-	void GoToAssetInContentBrowser( TWeakObjectPtr<UObject> Object )
-	{
-		if( Object.IsValid() )
-		{
-			TArray< UObject* > Objects;
-			Objects.Add( Object.Get() );
-			GEditor->SyncBrowserToObjects( Objects );
-		}
-	}
+	return PropertyHandle;
+}
+PRAGMA_ENABLE_OPTIMIZATION;
 
-private:
-	FMaterialListItem MaterialItem;
-	FOnMaterialChanged OnMaterialChanged;
-	FOnGenerateWidgetsForMaterial OnGenerateCustomNameWidgets;
-	FOnGenerateWidgetsForMaterial OnGenerateCustomMaterialWidgets;
-	FOnResetMaterialToDefaultClicked OnResetToDefaultClicked;
-	int32 MultipleMaterialCount;
-	bool bShowUsedTextures;
-};
+TSharedPtr<SWidget> FMaterialItemView::GetGlobalRowExtensionWidget(IDetailLayoutBuilder& InDetailBuilder, UActorComponent* InCurrentComponent) const
+{
+	const TSharedPtr<IPropertyHandle> MaterialPropertyHandle = GetMaterialHandle(InDetailBuilder, InCurrentComponent);
 
+	if (!MaterialPropertyHandle)
+	{
+		return SNullWidget::NullWidget;
+	}
+		
+	FSlimHorizontalToolBarBuilder ToolbarBuilder(TSharedPtr<FUICommandList>(), FMultiBoxCustomization::None);
+	FPropertyEditorModule& PropertyEditorModule = FModuleManager::Get().GetModuleChecked<FPropertyEditorModule>("PropertyEditor");
+
+	// Add Extension buttons to Row Generator
+	TArray<FPropertyRowExtensionButton> ExtensionButtons;
+
+	// Add custom property extensions
+	FOnGenerateGlobalRowExtensionArgs RowExtensionArgs;
+	RowExtensionArgs.PropertyHandle = MaterialPropertyHandle;
+	PropertyEditorModule.GetGlobalRowExtensionDelegate().Broadcast(RowExtensionArgs, ExtensionButtons);
+
+	// Build extension toolbar 
+	ToolbarBuilder.SetLabelVisibility(EVisibility::Collapsed);
+	ToolbarBuilder.SetStyle(&FAppStyle::Get(), "DetailsView.ExtensionToolBar");
+	for (const FPropertyRowExtensionButton& Extension : ExtensionButtons)
+	{
+		ToolbarBuilder.AddToolBarButton(Extension.UIAction, NAME_None, Extension.Label, Extension.ToolTip, Extension.Icon);
+	}
+		
+	return ToolbarBuilder.MakeWidget();
+}
 
 FMaterialList::FMaterialList(IDetailLayoutBuilder& InDetailLayoutBuilder, FMaterialListDelegates& InMaterialListDelegates, const TArray<FAssetData>& InOwnerAssetDataArray, bool bInAllowCollapse, bool bInShowUsedTextures)
 	: MaterialListDelegates( InMaterialListDelegates )
@@ -479,7 +615,7 @@ void FMaterialList::GenerateChildContent( IDetailChildrenBuilder& ChildrenBuilde
 					FDetailWidgetRow& ChildRow = ChildrenBuilder.AddCustomRow( FText::GetEmpty() );
 					ChildRow.RowTag(RowTagName);
 
-					AddMaterialItem( ChildRow, CurrentSlot, FMaterialListItem( NULL, CurrentSlot, true ), !bDisplayAllMaterialsInSlot );
+					AddMaterialItem( ChildRow, CurrentSlot, FMaterialListItem( NULL, CurrentSlot, true ), !bDisplayAllMaterialsInSlot, Material.CurrentComponent.Get() );
 				}
 				else
 				{
@@ -494,7 +630,7 @@ void FMaterialList::GenerateChildContent( IDetailChildrenBuilder& ChildrenBuilde
 				FDetailWidgetRow& ChildRow = ChildrenBuilder.AddCustomRow( Material.Material.IsValid()? FText::FromString(Material.Material->GetName()) : FText::GetEmpty() );
 				ChildRow.RowTag(RowTagName);
 
-				AddMaterialItem( ChildRow, CurrentSlot, Material, !bDisplayAllMaterialsInSlot );
+				AddMaterialItem( ChildRow, CurrentSlot, Material, !bDisplayAllMaterialsInSlot, Material.CurrentComponent.Get() );
 			}
 		}
 	}
@@ -567,7 +703,7 @@ void FMaterialList::OnPasteMaterialItem(int32 CurrentSlot)
 	}
 }
 			
-void FMaterialList::AddMaterialItem( FDetailWidgetRow& Row, int32 CurrentSlot, const FMaterialListItem& Item, bool bDisplayLink )
+void FMaterialList::AddMaterialItem( FDetailWidgetRow& Row, int32 CurrentSlot, const FMaterialListItem& Item, bool bDisplayLink, UActorComponent* InActorComponent)
 {
 	uint32 NumMaterials = MaterialListBuilder->GetNumMaterialsInSlot(CurrentSlot);
 
@@ -593,7 +729,7 @@ void FMaterialList::AddMaterialItem( FDetailWidgetRow& Row, int32 CurrentSlot, c
 	}
 	else
 	{
-		RightSideContent = NewView->CreateValueContent( DetailLayoutBuilder.GetThumbnailPool(), OwnerAssetDataArray );
+		RightSideContent = NewView->CreateValueContent( DetailLayoutBuilder, OwnerAssetDataArray, InActorComponent );
 		ViewedMaterials.Add( NewView );
 	}
 
diff --git a/Engine/Source/Editor/PropertyEditor/Public/MaterialList.h b/Engine/Source/Editor/PropertyEditor/Public/MaterialList.h
--- a/Engine/Source/Editor/PropertyEditor/Public/MaterialList.h
+++ b/Engine/Source/Editor/PropertyEditor/Public/MaterialList.h
@@ -11,6 +11,7 @@ class FMaterialListBuilder;
 class IDetailChildrenBuilder;
 class IDetailLayoutBuilder;
 class IMaterialListBuilder;
+class UActorComponent;
 
 /**
  * Delegate called when we need to get new materials for the list
@@ -90,11 +91,11 @@ public:
 	 * @param SlotIndex The slot (usually mesh element index) where the material is located on the component.
 	 * @param Material The material being used.
 	 * @param bCanBeReplced Whether or not the material can be replaced by a user.
+	 * @param InCurrentComponent The current component using the current material
 	 */
-	virtual void AddMaterial( uint32 SlotIndex, UMaterialInterface* Material, bool bCanBeReplaced ) = 0;
+	virtual void AddMaterial( uint32 SlotIndex, UMaterialInterface* Material, bool bCanBeReplaced, UActorComponent* InCurrentComponent = nullptr) = 0;
 };
 
-
 /**
  * A Material item in a material list slot
  */
@@ -109,10 +110,14 @@ struct FMaterialListItem
 	/** Whether or not this material can be replaced by a new material */
 	bool bCanBeReplaced;
 
-	FMaterialListItem( UMaterialInterface* InMaterial = NULL, uint32 InSlotIndex = 0, bool bInCanBeReplaced = false )
+	/** The current component using the current material */
+	TWeakObjectPtr<UActorComponent> CurrentComponent;
+
+	FMaterialListItem( UMaterialInterface* InMaterial = NULL, uint32 InSlotIndex = 0, bool bInCanBeReplaced = false, UActorComponent* InCurrentComponent = nullptr)
 		: Material( InMaterial )
 		, SlotIndex( InSlotIndex )
 		, bCanBeReplaced( bInCanBeReplaced )
+		, CurrentComponent(InCurrentComponent)
 	{}
 
 	friend uint32 GetTypeHash( const FMaterialListItem& InItem )
@@ -131,21 +136,157 @@ struct FMaterialListItem
 	}
 };
 
+/**
+ * A view of a single item in an FMaterialList
+ */
+class PROPERTYEDITOR_API FMaterialItemView : public TSharedFromThis<FMaterialItemView>
+{
+public:
+	/** Virtual destructor */
+	virtual ~FMaterialItemView() = default;
+	
+	/**
+	* Creates a new instance of this class
+	*
+	* @param Material				The material to view
+	* @param InOnMaterialChanged	Delegate for when the material changes
+	*/
+	static TSharedRef<FMaterialItemView> Create(
+		const FMaterialListItem& Material, 
+		FOnMaterialChanged InOnMaterialChanged,
+		FOnGenerateWidgetsForMaterial InOnGenerateNameWidgetsForMaterial, 
+		FOnGenerateWidgetsForMaterial InOnGenerateWidgetsForMaterial, 
+		FOnResetMaterialToDefaultClicked InOnResetToDefaultClicked,
+		int32 InMultipleMaterialCount,
+		bool bShowUsedTextures);
+
+	/**
+	 * Create a Name Content for Material Item Widget
+	 */
+	virtual TSharedRef<SWidget> CreateNameContent();
+
+	/**
+	 * Create a Value Content for Material Item Widget
+	 */
+	virtual TSharedRef<SWidget> CreateValueContent(IDetailLayoutBuilder& InDetailBuilder, const TArray<FAssetData>& OwnerAssetDataArray, UActorComponent* InActorComponent);
+
+	/**
+	 * Get Material item in a material list slot
+	 */
+	const FMaterialListItem& GetMaterialListItem() const { return MaterialItem; }
+
+	/**
+	 * Replace material in slot
+	 *
+	 * @param NewMaterial			The new material apply to slot
+	 * @param bReplaceAll			Whether it should be applied to all slots
+	 */
+	void ReplaceMaterial( UMaterialInterface* NewMaterial, bool bReplaceAll = false );
+
+	/**
+	 * Called to get the visibility of the reset to base button
+	 */
+	bool GetResetToBaseVisibility() const;
+
+	/**
+	 * Called when reset to base is clicked
+	 */
+	void OnResetToBaseClicked();
+
+private:
+	FMaterialItemView(const FMaterialListItem& InMaterial, 
+					FOnMaterialChanged& InOnMaterialChanged, 
+					FOnGenerateWidgetsForMaterial& InOnGenerateNameWidgets, 
+					FOnGenerateWidgetsForMaterial& InOnGenerateMaterialWidgets, 
+					FOnResetMaterialToDefaultClicked& InOnResetToDefaultClicked,
+					int32 InMultipleMaterialCount,
+					bool bInShowUsedTextures);
+
+	/**
+	 * Handler when the material asset has been changed
+	 */
+	void OnSetObject( const FAssetData& AssetData );
+
+	/**
+	 * @return Whether or not the textures menu is enabled
+	 */
+	bool IsTexturesMenuEnabled() const;
+
+	/**
+	 * Generate list of the textures for material selection list
+	 */
+	TSharedRef<SWidget> OnGetTexturesMenuForMaterial();
+
+	/**
+	 * On Get object path handler
+	 */
+	FString OnGetObjectPath() const;
+
+	/**
+	 * Finds the asset in the content browser
+	 */
+	void GoToAssetInContentBrowser( TWeakObjectPtr<UObject> Object );
+
+	/**
+	 * Get Material handle for given component
+	 *
+	 * @param InDetailBuilder The builder for laying custom details.
+	 * @param InCurrentComponent The current component using the current material.
+	 * @return Property Handle pointer
+	 */
+	TSharedPtr<IPropertyHandle> GetMaterialHandle(IDetailLayoutBuilder& InDetailBuilder, UActorComponent* InCurrentComponent) const;
+
+	/**
+	 * Get row extension widget
+	 *
+	 * @param InDetailBuilder The builder for laying custom details.
+	 * @param InCurrentComponent The current component using the current material.
+	 * @return Pointer to widget
+	 */
+	TSharedPtr<SWidget> GetGlobalRowExtensionWidget(IDetailLayoutBuilder& InDetailBuilder, UActorComponent* InCurrentComponent) const;
+
+private:
+	/** The Material item in a material list slot */
+	FMaterialListItem MaterialItem;
 
+	/** Material Changed delegate */
+	FOnMaterialChanged OnMaterialChanged;
+
+	/** Generate custom name delegate */
+	FOnGenerateWidgetsForMaterial OnGenerateCustomNameWidgets;
+
+	/** Generate custom material widget delegate */
+	FOnGenerateWidgetsForMaterial OnGenerateCustomMaterialWidgets;
+
+	/** Reset to default delegate */
+	FOnResetMaterialToDefaultClicked OnResetToDefaultClicked;
+
+	/** Number of material count */
+	int32 MultipleMaterialCount;
+
+	/** Whether it should show textures */
+	bool bShowUsedTextures;
+};
+
+/**
+ * Custom Node Builder for Material List
+ */
 class FMaterialList
 	: public IDetailCustomNodeBuilder
 	, public TSharedFromThis<FMaterialList>
 {
 public:
+	/** Add Extra widgets for bottom material value field. */
+	DECLARE_MULTICAST_DELEGATE_FourParams(FOnAddMaterialItemViewExtraBottomWidget,
+									const TSharedRef<FMaterialItemView>& /* InMaterialItemView */,
+									UActorComponent* /* InCurrentComponent */,
+									IDetailLayoutBuilder& /*InDetailBuilder*/,
+									TArray<TSharedPtr<SWidget>>& /* OutExtensions */);
+	PROPERTYEDITOR_API static FOnAddMaterialItemViewExtraBottomWidget OnAddMaterialItemViewExtraBottomWidget;
+	
 	PROPERTYEDITOR_API FMaterialList( IDetailLayoutBuilder& InDetailLayoutBuilder, FMaterialListDelegates& MaterialListDelegates, const TArray<FAssetData>& InOwnerAssetDataArray, bool bInAllowCollapse = false, bool bInShowUsedTextures = true);
 
-	/**
-	 * @return true if materials are being displayed.                                                          
-	 */
-	bool IsDisplayingMaterials() const { return true; }
-
 private:
-
 	/**
 	 * Called when a user expands all materials in a slot.
 	 *
@@ -172,12 +313,13 @@ private:
 	/**
 	 * Adds a new material item to the list
 	 *
-	 * @param Row			The row to add the item to
-	 * @param CurrentSlot	The slot id of the material
-	 * @param Item			The material item to add
-	 * @param bDisplayLink	If a link to the material should be displayed instead of the actual item (for multiple materials)
+	 * @param Row					The row to add the item to
+	 * @param CurrentSlot			The slot id of the material
+	 * @param Item					The material item to add
+	 * @param bDisplayLink			If a link to the material should be displayed instead of the actual item (for multiple materials)
+	 * @param InCurrentComponent	The current component using the current material
 	 */
-	void AddMaterialItem(FDetailWidgetRow& Row, int32 CurrentSlot, const FMaterialListItem& Item, bool bDisplayLink);
+	void AddMaterialItem(FDetailWidgetRow& Row, int32 CurrentSlot, const FMaterialListItem& Item, bool bDisplayLink, UActorComponent* InCurrentComponent = nullptr);
 
 private:
 	bool OnCanCopyMaterialList() const;
diff --git a/Engine/Source/Runtime/Engine/Classes/Components/DecalComponent.h b/Engine/Source/Runtime/Engine/Classes/Components/DecalComponent.h
--- a/Engine/Source/Runtime/Engine/Classes/Components/DecalComponent.h
+++ b/Engine/Source/Runtime/Engine/Classes/Components/DecalComponent.h
@@ -175,6 +175,9 @@ public:
 	//~ End UObject Interface
 
 	//~ Begin USceneComponent Interface
+#if WITH_EDITOR
+	virtual bool GetMaterialPropertyPath(int32 ElementIndex, UObject*& OutOwner, FString& OutPropertyPath) override;
+#endif // WITH_EDITOR
 	virtual FBoxSphereBounds CalcBounds(const FTransform& LocalToWorld) const override;
 	//~ End USceneComponent Interface
 
diff --git a/Engine/Source/Runtime/Engine/Classes/Components/LightComponent.h b/Engine/Source/Runtime/Engine/Classes/Components/LightComponent.h
--- a/Engine/Source/Runtime/Engine/Classes/Components/LightComponent.h
+++ b/Engine/Source/Runtime/Engine/Classes/Components/LightComponent.h
@@ -441,6 +441,12 @@ protected:
 	virtual void DestroyRenderState_Concurrent() override;
 	//~ Begin UActorComponent Interface
 
+	//~ Begin USceneComponent Interface
+#if WITH_EDITOR
+	virtual bool GetMaterialPropertyPath(int32 ElementIndex, UObject*& OutOwner, FString& OutPropertyPath) override;
+#endif // WITH_EDITOR
+	//~ End USceneComponent Interface
+
 public:
 	virtual void InvalidateLightingCacheDetailed(bool bInvalidateBuildEnqueuedLighting, bool bTranslationOnly) override;
 
diff --git a/Engine/Source/Runtime/Engine/Classes/Components/MaterialBillboardComponent.h b/Engine/Source/Runtime/Engine/Classes/Components/MaterialBillboardComponent.h
--- a/Engine/Source/Runtime/Engine/Classes/Components/MaterialBillboardComponent.h
+++ b/Engine/Source/Runtime/Engine/Classes/Components/MaterialBillboardComponent.h
@@ -88,4 +88,10 @@ class ENGINE_API UMaterialBillboardComponent : public UPrimitiveComponent
 	virtual void SetMaterial(int32 ElementIndex, class UMaterialInterface* Material) override;
 	virtual void GetUsedMaterials(TArray<UMaterialInterface*>& OutMaterials, bool bGetDebugMaterials = false) const override;
 	//~ End UPrimitiveComponent Interface
+
+	//~ Begin USceneComponent Interface
+#if WITH_EDITOR
+	virtual bool GetMaterialPropertyPath(int32 ElementIndex, UObject*& OutOwner, FString& OutPropertyPath) override;
+#endif // WITH_EDITOR
+	//~ End USceneComponent Interface
 };
diff --git a/Engine/Source/Runtime/Engine/Classes/Components/SceneComponent.h b/Engine/Source/Runtime/Engine/Classes/Components/SceneComponent.h
--- a/Engine/Source/Runtime/Engine/Classes/Components/SceneComponent.h
+++ b/Engine/Source/Runtime/Engine/Classes/Components/SceneComponent.h
@@ -804,6 +804,22 @@ public:
 	UFUNCTION(BlueprintCallable, Category="Rendering")
 	virtual bool IsVisible() const;
 
+#if WITH_EDITOR
+	/**
+	 * Returns full material property path and UObject owner property object
+	 * Path examples:
+	 * Material property path with array element and inner struct Materials[0].InnerStruct.Material
+	 * Material property path with array element Materials[0]
+	 * Simple material property path Materials
+	 * 
+	 * @param ElementIndex		- The element to access the material of.
+	 * @param OutOwner			- Property UObject owner.
+	 * @param OutPropertyPath	- Full material property path.
+	 * @return true if that was successfully resolved and component has material
+	 */
+	virtual bool GetMaterialPropertyPath(int32 ElementIndex, UObject*& OutOwner, FString& OutPropertyPath);
+#endif // WITH_EDITOR
+
 protected:
 	/**
 	 * Overridable internal function to respond to changes in the visibility of the component.
diff --git a/Engine/Source/Runtime/Engine/Classes/Components/SkinnedMeshComponent.h b/Engine/Source/Runtime/Engine/Classes/Components/SkinnedMeshComponent.h
--- a/Engine/Source/Runtime/Engine/Classes/Components/SkinnedMeshComponent.h
+++ b/Engine/Source/Runtime/Engine/Classes/Components/SkinnedMeshComponent.h
@@ -865,6 +865,9 @@ protected:
 
 public:
 	//~ Begin USceneComponent Interface
+#if WITH_EDITOR
+	virtual bool GetMaterialPropertyPath(int32 ElementIndex, UObject*& OutOwner, FString& OutPropertyPath) override;
+#endif // WITH_EDITOR
 	virtual FBoxSphereBounds CalcBounds(const FTransform& LocalToWorld) const override;
 	virtual FTransform GetSocketTransform(FName InSocketName, ERelativeTransformSpace TransformSpace = RTS_World) const override;
 	virtual bool DoesSocketExist(FName InSocketName) const override;
diff --git a/Engine/Source/Runtime/Engine/Classes/Components/StaticMeshComponent.h b/Engine/Source/Runtime/Engine/Classes/Components/StaticMeshComponent.h
--- a/Engine/Source/Runtime/Engine/Classes/Components/StaticMeshComponent.h
+++ b/Engine/Source/Runtime/Engine/Classes/Components/StaticMeshComponent.h
@@ -469,6 +469,9 @@ public:
 	//~ End UObject Interface.
 
 	//~ Begin USceneComponent Interface
+#if WITH_EDITOR
+	virtual bool GetMaterialPropertyPath(int32 ElementIndex, UObject*& OutOwner, FString& OutPropertyPath) override;
+#endif // WITH_EDITOR
 	virtual FBoxSphereBounds CalcBounds(const FTransform& LocalToWorld) const override;
 	virtual bool HasAnySockets() const override;
 	virtual void QuerySupportedSockets(TArray<FComponentSocketDescription>& OutSockets) const override;
diff --git a/Engine/Source/Runtime/Engine/Classes/Components/TextRenderComponent.h b/Engine/Source/Runtime/Engine/Classes/Components/TextRenderComponent.h
--- a/Engine/Source/Runtime/Engine/Classes/Components/TextRenderComponent.h
+++ b/Engine/Source/Runtime/Engine/Classes/Components/TextRenderComponent.h
@@ -167,6 +167,9 @@ class ENGINE_API UTextRenderComponent : public UPrimitiveComponent
 	//~ End UPrimitiveComponent Interface.
 
 	//~ Begin USceneComponent Interface.
+#if WITH_EDITOR
+	virtual bool GetMaterialPropertyPath(int32 ElementIndex, UObject*& OutOwner, FString& OutPropertyPath) override;
+#endif // WITH_EDITOR
 	virtual FBoxSphereBounds CalcBounds(const FTransform& LocalToWorld) const override;
 	// Note: this is required because GetRenderMatrix is overridden, and therefore the updated world-space bounds must be based on that (not the component transform)
 	virtual void UpdateBounds() override;
diff --git a/Engine/Source/Runtime/Engine/Classes/Particles/ParticleSystemComponent.h b/Engine/Source/Runtime/Engine/Classes/Particles/ParticleSystemComponent.h
--- a/Engine/Source/Runtime/Engine/Classes/Particles/ParticleSystemComponent.h
+++ b/Engine/Source/Runtime/Engine/Classes/Particles/ParticleSystemComponent.h
@@ -1370,7 +1370,7 @@ public:
 	//~ End UObject Interface.
 
 	//Begin UPrimitiveComponent Interface
-	virtual int32 GetNumMaterials() const override; 
+	virtual int32 GetNumMaterials() const override;
 	virtual UMaterialInterface* GetMaterial(int32 ElementIndex) const override;
 	virtual void SetMaterial(int32 ElementIndex, UMaterialInterface* Material) override;
 	virtual FBoxSphereBounds CalcBounds(const FTransform& LocalToWorld) const override;
@@ -1385,6 +1385,9 @@ protected:
 	virtual bool ShouldActivate() const override;
 
 public:
+#if WITH_EDITOR
+	virtual bool GetMaterialPropertyPath(int32 ElementIndex, UObject*& OutOwner, FString& OutPropertyPath) override;
+#endif // WITH_EDITOR
 	virtual void Activate(bool bReset=false) override;
 	virtual void Deactivate() override;
 	virtual void DeactivateImmediate() override;
diff --git a/Engine/Source/Runtime/Engine/Private/Components/DecalComponent.cpp b/Engine/Source/Runtime/Engine/Private/Components/DecalComponent.cpp
--- a/Engine/Source/Runtime/Engine/Private/Components/DecalComponent.cpp
+++ b/Engine/Source/Runtime/Engine/Private/Components/DecalComponent.cpp
@@ -132,6 +132,20 @@ bool UDecalComponent::IsPostLoadThreadSafe() const
 	return true;
 }
 
+#if WITH_EDITOR
+bool UDecalComponent::GetMaterialPropertyPath(int32 ElementIndex, UObject*& OutOwner, FString& OutPropertyPath)
+{
+	if(ElementIndex == 0)
+	{
+		OutOwner = this;
+		OutPropertyPath = GET_MEMBER_NAME_STRING_CHECKED(UDecalComponent, DecalMaterial);
+		return true;
+	}
+
+	return false;
+}
+#endif // WITH_EDITOR
+
 void UDecalComponent::SetLifeSpan(const float LifeSpan)
 {
 	if (LifeSpan > 0.f)
diff --git a/Engine/Source/Runtime/Engine/Private/Components/LightComponent.cpp b/Engine/Source/Runtime/Engine/Private/Components/LightComponent.cpp
--- a/Engine/Source/Runtime/Engine/Private/Components/LightComponent.cpp
+++ b/Engine/Source/Runtime/Engine/Private/Components/LightComponent.cpp
@@ -903,6 +903,20 @@ void ULightComponent::DestroyRenderState_Concurrent()
 	bAddedToSceneVisible = false;
 }
 
+#if WITH_EDITOR
+bool ULightComponent::GetMaterialPropertyPath(int32 ElementIndex, UObject*& OutOwner, FString& OutPropertyPath)
+{
+	if (ElementIndex == 0)
+	{
+		OutOwner = this;
+		OutPropertyPath = GET_MEMBER_NAME_STRING_CHECKED(ULightComponent, LightFunctionMaterial);
+		return true;
+	}
+
+	return false;
+}
+#endif // WITH_EDITOR
+
 /** Set brightness of the light */
 void ULightComponent::SetIntensity(float NewIntensity)
 {
diff --git a/Engine/Source/Runtime/Engine/Private/Components/MaterialBillboardComponent.cpp b/Engine/Source/Runtime/Engine/Private/Components/MaterialBillboardComponent.cpp
--- a/Engine/Source/Runtime/Engine/Private/Components/MaterialBillboardComponent.cpp
+++ b/Engine/Source/Runtime/Engine/Private/Components/MaterialBillboardComponent.cpp
@@ -318,6 +318,21 @@ FBoxSphereBounds UMaterialBillboardComponent::CalcBounds(const FTransform& Local
 	return FBoxSphereBounds(LocalToWorld.GetLocation(),FVector(BoundsSize,BoundsSize,BoundsSize),FMath::Sqrt(3.0f * FMath::Square(BoundsSize)));
 }
 
+#if WITH_EDITOR
+bool UMaterialBillboardComponent::GetMaterialPropertyPath(int32 ElementIndex, UObject*& OutOwner, FString& OutPropertyPath)
+{
+	if (Elements.IsValidIndex(ElementIndex))
+	{
+		OutOwner = this;
+		OutPropertyPath = FString::Printf(TEXT("%s[%d].%s"), GET_MEMBER_NAME_STRING_CHECKED(UMaterialBillboardComponent, Elements), ElementIndex, GET_MEMBER_NAME_STRING_CHECKED(FMaterialSpriteElement, Material));
+	
+		return true;
+	}
+
+	return false;
+}
+#endif // WITH_EDITOR
+
 void UMaterialBillboardComponent::AddElement(
 	class UMaterialInterface* Material,
 	UCurveFloat* DistanceToOpacityCurve,
diff --git a/Engine/Source/Runtime/Engine/Private/Components/SceneComponent.cpp b/Engine/Source/Runtime/Engine/Private/Components/SceneComponent.cpp
--- a/Engine/Source/Runtime/Engine/Private/Components/SceneComponent.cpp
+++ b/Engine/Source/Runtime/Engine/Private/Components/SceneComponent.cpp
@@ -3030,6 +3030,14 @@ bool USceneComponent::IsVisible() const
 	return (GetVisibleFlag() && (!CachedLevelCollection || CachedLevelCollection->IsVisible())); 
 }
 
+#if WITH_EDITOR
+bool USceneComponent::GetMaterialPropertyPath(int32 ElementIndex, UObject*& OutOwner, FString& OutPropertyPath)
+{
+	// Should be overriden in inherited classes
+	return false;
+}
+#endif // WITH_EDITOR
+
 void USceneComponent::OnVisibilityChanged()
 {
 	MarkRenderStateDirty();
@@ -3080,6 +3088,7 @@ void USceneComponent::OnHiddenInGameChanged()
 	MarkRenderStateDirty();
 }
 
+
 void USceneComponent::SetHiddenInGame(const bool bNewHiddenGame, const USceneComponent::EVisibilityPropagation PropagateToChildren)
 {
 	bool bRecurseChildren = (PropagateToChildren == EVisibilityPropagation::Propagate);
diff --git a/Engine/Source/Runtime/Engine/Private/Components/SkinnedMeshComponent.cpp b/Engine/Source/Runtime/Engine/Private/Components/SkinnedMeshComponent.cpp
--- a/Engine/Source/Runtime/Engine/Private/Components/SkinnedMeshComponent.cpp
+++ b/Engine/Source/Runtime/Engine/Private/Components/SkinnedMeshComponent.cpp
@@ -2388,6 +2388,26 @@ bool USkinnedMeshComponent::UpdateOverlapsImpl(const TOverlapArrayView* PendingO
 	return USceneComponent::UpdateOverlapsImpl(PendingOverlaps, bDoNotifies, OverlapsAtEndLocation);
 }
 
+#if WITH_EDITOR
+bool USkinnedMeshComponent::GetMaterialPropertyPath(int32 ElementIndex, UObject*& OutOwner, FString& OutPropertyPath)
+{
+	if(OverrideMaterials.IsValidIndex(ElementIndex))
+	{
+		OutOwner = this;
+		OutPropertyPath = FString::Printf(TEXT("%s[%d]"), GET_MEMBER_NAME_STRING_CHECKED(UMeshComponent, OverrideMaterials), ElementIndex);
+		return true;
+	}
+	if (SkeletalMesh && SkeletalMesh->GetMaterials().IsValidIndex(ElementIndex))
+	{
+		OutOwner = SkeletalMesh;
+		OutPropertyPath = FString::Printf(TEXT("%s[%d].%s"), *USkeletalMesh::GetMaterialsMemberName().ToString(), ElementIndex, GET_MEMBER_NAME_STRING_CHECKED(FSkeletalMaterial, MaterialInterface));
+		return true;
+	}
+
+	return false;
+}
+#endif // WITH_EDITOR
+
 void USkinnedMeshComponent::TransformToBoneSpace(FName BoneName, FVector InPosition, FRotator InRotation, FVector& OutPosition, FRotator& OutRotation) const
 {
 	int32 BoneIndex = GetBoneIndex(BoneName);
diff --git a/Engine/Source/Runtime/Engine/Private/Components/StaticMeshComponent.cpp b/Engine/Source/Runtime/Engine/Private/Components/StaticMeshComponent.cpp
--- a/Engine/Source/Runtime/Engine/Private/Components/StaticMeshComponent.cpp
+++ b/Engine/Source/Runtime/Engine/Private/Components/StaticMeshComponent.cpp
@@ -2503,6 +2503,28 @@ void UStaticMeshComponent::GetUsedMaterials(TArray<UMaterialInterface*>& OutMate
 	}
 }
 
+#if WITH_EDITOR
+bool UStaticMeshComponent::GetMaterialPropertyPath(int32 ElementIndex, UObject*& OutOwner, FString& OutPropertyPath)
+{
+	if(OverrideMaterials.IsValidIndex(ElementIndex))
+	{				
+		OutOwner = this;
+		OutPropertyPath = FString::Printf(TEXT("%s[%d]"), GET_MEMBER_NAME_STRING_CHECKED(UMeshComponent, OverrideMaterials), ElementIndex);
+
+		return true;
+	}
+	if (GetStaticMesh())
+	{
+		OutOwner = GetStaticMesh();
+		OutPropertyPath = FString::Printf(TEXT("%s[%d].%s"), *UStaticMesh::GetStaticMaterialsName().ToString(), ElementIndex, GET_MEMBER_NAME_STRING_CHECKED(FStaticMaterial, MaterialInterface));
+
+		return true;
+	}
+
+	return false;
+}
+#endif // WITH_EDITOR
+
 int32 UStaticMeshComponent::GetBlueprintCreatedComponentIndex() const
 {
 	int32 ComponentIndex = 0;
diff --git a/Engine/Source/Runtime/Engine/Private/Components/TextRenderComponent.cpp b/Engine/Source/Runtime/Engine/Private/Components/TextRenderComponent.cpp
--- a/Engine/Source/Runtime/Engine/Private/Components/TextRenderComponent.cpp
+++ b/Engine/Source/Runtime/Engine/Private/Components/TextRenderComponent.cpp
@@ -1107,6 +1107,20 @@ int32 UTextRenderComponent::GetNumMaterials() const
 	return 1;
 }
 
+#if WITH_EDITOR
+bool UTextRenderComponent::GetMaterialPropertyPath(int32 ElementIndex, UObject*& OutOwner, FString& OutPropertyPath)
+{
+	if (ElementIndex == 0)
+	{
+		OutOwner = this;
+		OutPropertyPath = GET_MEMBER_NAME_STRING_CHECKED(UTextRenderComponent, TextMaterial);
+		return true;
+	}
+
+	return false;
+}
+#endif // WITH_EDITOR
+
 void UTextRenderComponent::SetMaterial(int32 ElementIndex, UMaterialInterface* InMaterial)
 {
 	if (ElementIndex == 0)
diff --git a/Engine/Source/Runtime/Engine/Private/Particles/ParticleComponents.cpp b/Engine/Source/Runtime/Engine/Private/Particles/ParticleComponents.cpp
--- a/Engine/Source/Runtime/Engine/Private/Particles/ParticleComponents.cpp
+++ b/Engine/Source/Runtime/Engine/Private/Particles/ParticleComponents.cpp
@@ -4444,6 +4444,34 @@ int32 UParticleSystemComponent::GetNumMaterials() const
 	return 0;
 }
 
+#if WITH_EDITOR
+bool UParticleSystemComponent::GetMaterialPropertyPath(int32 ElementIndex, UObject*& OutOwner, FString& OutPropertyPath)
+{
+	if (EmitterMaterials.IsValidIndex(ElementIndex))
+	{
+		OutOwner = this;
+		OutPropertyPath = FString::Printf(TEXT("%s[%d]"), GET_MEMBER_NAME_STRING_CHECKED(UParticleSystemComponent, EmitterMaterials), ElementIndex);
+		return true;
+	}
+	if (Template && Template->Emitters.IsValidIndex(ElementIndex))
+	{
+		UParticleEmitter* Emitter = Template->Emitters[ElementIndex];
+		if (Emitter && Emitter->LODLevels.Num() > 0)
+		{
+			UParticleLODLevel* EmitterLODLevel = Emitter->LODLevels[0];
+			if (EmitterLODLevel && EmitterLODLevel->RequiredModule)
+			{
+				OutOwner = EmitterLODLevel->RequiredModule;
+				OutPropertyPath = GET_MEMBER_NAME_STRING_CHECKED(UParticleModuleRequired, Material);
+				return true;
+			}
+		}
+	}
+
+	return false;
+}
+#endif // WITH_EDITOR
+
 UMaterialInterface* UParticleSystemComponent::GetMaterial(int32 ElementIndex) const
 {
 	if (EmitterMaterials.IsValidIndex(ElementIndex) && EmitterMaterials[ElementIndex] != NULL)
diff --git a/Engine/Source/Runtime/Landscape/Classes/LandscapeComponent.h b/Engine/Source/Runtime/Landscape/Classes/LandscapeComponent.h
--- a/Engine/Source/Runtime/Landscape/Classes/LandscapeComponent.h
+++ b/Engine/Source/Runtime/Landscape/Classes/LandscapeComponent.h
@@ -779,6 +779,9 @@ public:
 	//~ End UPrimitiveComponent Interface.
 
 	//~ Begin USceneComponent Interface.
+#if WITH_EDITOR
+	virtual bool GetMaterialPropertyPath(int32 ElementIndex, UObject*& OutOwner, FString& OutPropertyPath) override;
+#endif // WITH_EDITOR
 	virtual void DestroyComponent(bool bPromoteChildren = false) override;
 	virtual FBoxSphereBounds CalcBounds(const FTransform& LocalToWorld) const override;
 	//~ End USceneComponent Interface.
diff --git a/Engine/Source/Runtime/Landscape/Private/LandscapeEdit.cpp b/Engine/Source/Runtime/Landscape/Private/LandscapeEdit.cpp
--- a/Engine/Source/Runtime/Landscape/Private/LandscapeEdit.cpp
+++ b/Engine/Source/Runtime/Landscape/Private/LandscapeEdit.cpp
@@ -533,6 +533,30 @@ int32 ULandscapeComponent::GetNumMaterials() const
 	return 1;
 }
 
+#if WITH_EDITOR
+bool ULandscapeComponent::GetMaterialPropertyPath(int32 ElementIndex, UObject*& OutOwner, FString& OutPropertyPath)
+{
+	if (ElementIndex == 0)
+	{
+		if (OverrideMaterial)
+		{
+			OutOwner = this;
+			OutPropertyPath = GET_MEMBER_NAME_STRING_CHECKED(ULandscapeComponent, OverrideMaterial);
+	
+			return true;
+		}
+		if (ALandscapeProxy* Proxy = GetLandscapeProxy())
+		{
+			OutOwner = Proxy;
+			OutPropertyPath = GET_MEMBER_NAME_STRING_CHECKED(ALandscapeProxy, LandscapeHoleMaterial);
+			return true;
+		}
+	}
+
+	return false;
+}
+#endif // WITH_EDITOR
+
 class UMaterialInterface* ULandscapeComponent::GetMaterial(int32 ElementIndex) const
 {
 	if (ensure(ElementIndex == 0))
