Fix issues preflight issues for Dynamic Material Instance extenssion

diff --git a/Engine/Plugins/Experimental/VirtualProductionUtilities/Source/VPMaterialsEditor/Private/SMaterialDynamicParametersPanelWidget.cpp b/Engine/Plugins/Experimental/VirtualProductionUtilities/Source/VPMaterialsEditor/Private/SMaterialDynamicParametersPanelWidget.cpp
--- a/Engine/Plugins/Experimental/VirtualProductionUtilities/Source/VPMaterialsEditor/Private/SMaterialDynamicParametersPanelWidget.cpp
+++ b/Engine/Plugins/Experimental/VirtualProductionUtilities/Source/VPMaterialsEditor/Private/SMaterialDynamicParametersPanelWidget.cpp
@@ -471,7 +471,7 @@ void SMaterialDynamicParametersOverviewTree::OnExpansionChanged(TSharedPtr<FProp
 		MaterialInstance->ParameterOverviewExpansion.Emplace(Item->NodeKey, bIsExpanded);
 	}
 	// Expand any children that are also expanded
-	for (const TSharedPtr<FPropertySortedParamData> Child : Item->Children)
+	for (TSharedPtr<FPropertySortedParamData> Child : Item->Children)
 	{
 		bool* ChildExpansionValue = MaterialInstance->ParameterOverviewExpansion.Find(Child->NodeKey);
 		if (ChildExpansionValue != nullptr && *ChildExpansionValue == true)
@@ -588,7 +588,7 @@ void SMaterialDynamicParametersOverviewTree::ShowSubParameters()
 		GroupProperty->NodeKey = FString::FromInt(GroupProperty->ParameterInfo.Index) + FString::FromInt(GroupProperty->ParameterInfo.Association) + Property.GroupName.ToString();
 
 		bool bAddNewGroup = true;
-		for (const TSharedPtr<struct FPropertySortedParamData> GroupChild : SortedParameters)
+		for (TSharedPtr<struct FPropertySortedParamData> GroupChild : SortedParameters)
 		{
 			if (GroupChild->NodeKey == GroupProperty->NodeKey)
 			{
diff --git a/Engine/Source/Editor/PropertyEditor/Private/MaterialList.cpp b/Engine/Source/Editor/PropertyEditor/Private/MaterialList.cpp
--- a/Engine/Source/Editor/PropertyEditor/Private/MaterialList.cpp
+++ b/Engine/Source/Editor/PropertyEditor/Private/MaterialList.cpp
@@ -153,7 +153,7 @@ namespace ValueExtender
 		// Execute all delegates and loop through all results widgets
 		InDelegate.Broadcast(InMaterialItemView, InCurrentComponent, InDetailBuilder, Widgets);
 
-		for (const TSharedPtr<SWidget> Widget : Widgets)
+		for (TSharedPtr<SWidget> Widget : Widgets)
 		{
 			WidgetContainer->AddSlot()
 			[
