- Changing fatal error when role doesn't exist to only be error and avoid crashing

diff --git a/Engine/Plugins/Experimental/VirtualProductionUtilities/Source/VPUtilities/Private/VPSettings.cpp b/Engine/Plugins/Experimental/VirtualProductionUtilities/Source/VPUtilities/Private/VPSettings.cpp
--- a/Engine/Plugins/Experimental/VirtualProductionUtilities/Source/VPUtilities/Private/VPSettings.cpp
+++ b/Engine/Plugins/Experimental/VirtualProductionUtilities/Source/VPUtilities/Private/VPSettings.cpp
@@ -24,7 +24,7 @@ UVPSettings::UVPSettings()
 			}
 			else
 			{
-				UE_LOG(LogVPUtilities, Fatal, TEXT("Role %s doesn't exist."), *Role);
+				UE_LOG(LogVPUtilities, Error, TEXT("Role %s doesn't exist."), *Role);
 			}
 		}
 	}
