Fix: Missing tooltips for Dynamic Material Instance buttons

diff --git a/Engine/Plugins/Experimental/VirtualProductionUtilities/Source/VPMaterialsEditor/Private/SMaterialDynamicWidgets.cpp b/Engine/Plugins/Experimental/VirtualProductionUtilities/Source/VPMaterialsEditor/Private/SMaterialDynamicWidgets.cpp
--- a/Engine/Plugins/Experimental/VirtualProductionUtilities/Source/VPMaterialsEditor/Private/SMaterialDynamicWidgets.cpp
+++ b/Engine/Plugins/Experimental/VirtualProductionUtilities/Source/VPMaterialsEditor/Private/SMaterialDynamicWidgets.cpp
@@ -40,6 +40,7 @@ void SMaterialDynamicView::Construct(const FArguments& InArgs, const TSharedRef<
 					[
 						SNew(SButton)
 						.Visibility(GetButtonVisibilityAttribute<UMaterialInstanceDynamic>())
+						.ToolTipText(LOCTEXT("Revert_Button_Tooltip", "Revert the Dynamic Material Instance back to the original Material Instance"))
 						.OnClicked(this, &SMaterialDynamicView::OnRevertButtonClicked)
 						[
 							SNew(STextBlock).Text(LOCTEXT("Revert_Button", "Revert"))
@@ -51,6 +52,7 @@ void SMaterialDynamicView::Construct(const FArguments& InArgs, const TSharedRef<
 					[
 						SNew(SButton)
 						.Visibility(GetButtonVisibilityAttribute<UMaterialInstanceDynamic>())
+						.ToolTipText(LOCTEXT("Reset_Button_Tooltip", "Reset the properties to the Material Instance default"))
 						.OnClicked(this, &SMaterialDynamicView::OnResetButtonClicked)
 						[
 							SNew(STextBlock).Text(LOCTEXT("Reset_Button", "Reset"))
@@ -62,6 +64,7 @@ void SMaterialDynamicView::Construct(const FArguments& InArgs, const TSharedRef<
 					[
 						SNew(SButton)
 						.Visibility(GetButtonVisibilityAttribute<UMaterialInstanceDynamic>())
+						.ToolTipText(LOCTEXT("CopyToOriginal_Button_Tooltip", "Copy and overwrite the parameters onto the original Material Instance"))
 						.OnClicked(this, &SMaterialDynamicView::OnCopyToOriginalButtonClicked)
 						[
 							SNew(STextBlock).Text(LOCTEXT("CopyToOriginal_Button", "Copy to Original"))
@@ -73,6 +76,7 @@ void SMaterialDynamicView::Construct(const FArguments& InArgs, const TSharedRef<
 					[
 						SNew(SButton)
 						.Visibility(CreateDynamicMaterialButtonVisibility)
+						.ToolTipText(LOCTEXT("CreateDynamicMaterial_Button_Tooltip", "Create a Dynamic Material Instance for this Material Instance and automatically sets it"))
 						.OnClicked(this, &SMaterialDynamicView::OnCreateDynamicMaterialButtonClicked)
 						[
 							SNew(STextBlock).Text(LOCTEXT("CreateDynamicMaterial_Button", "Create Dynamic Material"))
