Mark handful VirtualProduction plugin modules as UncookedOnly

diff --git a/Engine/Plugins/Experimental/VirtualCamera/VirtualCamera.uplugin b/Engine/Plugins/Experimental/VirtualCamera/VirtualCamera.uplugin
--- a/Engine/Plugins/Experimental/VirtualCamera/VirtualCamera.uplugin
+++ b/Engine/Plugins/Experimental/VirtualCamera/VirtualCamera.uplugin
@@ -17,7 +17,7 @@
 	"Modules": [
 		{
 			"Name": "VCamCore",
-			"Type": "Runtime",
+			"Type": "UncookedOnly",
 			"LoadingPhase": "Default",
 			"PlatformDenyList": [
 				"TVOS"
diff --git a/Engine/Plugins/Experimental/VirtualProductionUtilities/Source/VPUtilities/Private/VPFullScreenUserWidget.cpp b/Engine/Plugins/Experimental/VirtualProductionUtilities/Source/VPUtilities/Private/VPFullScreenUserWidget.cpp
--- a/Engine/Plugins/Experimental/VirtualProductionUtilities/Source/VPUtilities/Private/VPFullScreenUserWidget.cpp
+++ b/Engine/Plugins/Experimental/VirtualProductionUtilities/Source/VPUtilities/Private/VPFullScreenUserWidget.cpp
@@ -718,12 +718,12 @@ bool UVPFullScreenUserWidget::Display(UWorld* InWorld)
 			// If we are using Composure as our output, then send the WidgetRenderTarget to each one
 			if (CurrentDisplayType == EVPWidgetDisplayType::Composure)
 			{
-				static const FString TextureCompClassName("BP_TextureRTCompElement_C");
+				static const FName TextureCompClassName("BP_TextureRTCompElement_C");
 				static const FName TextureInputPropertyName("TextureRTInput");
 
 				for (ACompositingElement* Layer : PostProcessDisplayType.ComposureLayerTargets)
 				{
-					if (Layer && (Layer->GetClass()->GetName() == TextureCompClassName))
+					if (Layer && (Layer->GetClass()->GetFName() == TextureCompClassName))
 					{
 						FProperty* TextureInputProperty = Layer->GetClass()->FindPropertyByName(TextureInputPropertyName);
 						if (TextureInputProperty)
@@ -742,7 +742,7 @@ bool UVPFullScreenUserWidget::Display(UWorld* InWorld)
 					}
 					else if (Layer)
 					{
-						UE_LOG(LogVPUtilities, Warning, TEXT("VPFullScreenUserWidget - ComposureLayerTarget entry '%s' is not the correct class '%s'"), *Layer->GetName(), *TextureCompClassName);
+						UE_LOG(LogVPUtilities, Warning, TEXT("VPFullScreenUserWidget - ComposureLayerTarget entry '%s' is not the correct class '%s'"), *Layer->GetName(), *TextureCompClassName.ToString());
 					}
 				}
 			}
diff --git a/Engine/Plugins/Experimental/VirtualProductionUtilities/VirtualProductionUtilities.uplugin b/Engine/Plugins/Experimental/VirtualProductionUtilities/VirtualProductionUtilities.uplugin
--- a/Engine/Plugins/Experimental/VirtualProductionUtilities/VirtualProductionUtilities.uplugin
+++ b/Engine/Plugins/Experimental/VirtualProductionUtilities/VirtualProductionUtilities.uplugin
@@ -16,7 +16,7 @@
 	"Modules": [
 		{
 			"Name": "VPUtilities",
-			"Type": "Runtime",
+			"Type": "UncookedOnly",
 			"LoadingPhase": "PostDefault"
 		},
 		{
@@ -26,7 +26,7 @@
 		},
 		{
 			"Name": "VPBookmark",
-			"Type": "Runtime",
+			"Type": "UncookedOnly",
 			"LoadingPhase": "PostDefault"
 		},
 		{
diff --git a/Engine/Plugins/VirtualProduction/LevelSnapshots/LevelSnapshots.uplugin b/Engine/Plugins/VirtualProduction/LevelSnapshots/LevelSnapshots.uplugin
--- a/Engine/Plugins/VirtualProduction/LevelSnapshots/LevelSnapshots.uplugin
+++ b/Engine/Plugins/VirtualProduction/LevelSnapshots/LevelSnapshots.uplugin
@@ -18,12 +18,12 @@
 	"Modules": [
 		{
 			"Name": "LevelSnapshots",
-			"Type": "Runtime",
+			"Type": "UncookedOnly",
 			"LoadingPhase": "Default"
 		},
 		{
 			"Name": "LevelSnapshotFilters",
-			"Type": "Runtime",
+			"Type": "UncookedOnly",
 			"LoadingPhase": "Default"
 		},
 		{
@@ -38,7 +38,7 @@
 		},
         {
             "Name": "nDisplaySupport",
-            "Type": "Runtime",
+            "Type": "UncookedOnly",
             "LoadingPhase": "PostDefault",
             "PlatformAllowList": [
               "Win64",
diff --git a/Engine/Plugins/VirtualProduction/StageMonitoring/StageMonitoring.uplugin b/Engine/Plugins/VirtualProduction/StageMonitoring/StageMonitoring.uplugin
--- a/Engine/Plugins/VirtualProduction/StageMonitoring/StageMonitoring.uplugin
+++ b/Engine/Plugins/VirtualProduction/StageMonitoring/StageMonitoring.uplugin
@@ -26,7 +26,7 @@
 	"Modules": [
 		{
 			"Name": "StageMonitor",
-			"Type": "Runtime",
+			"Type": "UncookedOnly",
 			"LoadingPhase": "Default"
 		},
 		{
@@ -36,12 +36,12 @@
 		},
 		{
 			"Name": "StageDataProvider",
-			"Type": "Runtime",
+			"Type": "UncookedOnly",
 			"LoadingPhase": "Default"
 		},
 		{
 			"Name": "StageMonitorCommon",
-			"Type": "Runtime",
+			"Type": "UncookedOnly",
 			"LoadingPhase": "Default"
 		}
 	]
