Fix Jira issue, Reset Arrow.

diff --git a/Engine/Plugins/Experimental/VirtualProductionUtilities/Source/VPMaterialsEditor/Private/SMaterialDynamicParametersPanelWidget.cpp b/Engine/Plugins/Experimental/VirtualProductionUtilities/Source/VPMaterialsEditor/Private/SMaterialDynamicParametersPanelWidget.cpp
--- a/Engine/Plugins/Experimental/VirtualProductionUtilities/Source/VPMaterialsEditor/Private/SMaterialDynamicParametersPanelWidget.cpp
+++ b/Engine/Plugins/Experimental/VirtualProductionUtilities/Source/VPMaterialsEditor/Private/SMaterialDynamicParametersPanelWidget.cpp
@@ -186,7 +186,7 @@ private:
 
 	FSlateColor GetOuterBackgroundColor(TSharedPtr<FPropertySortedParamData> InParamData) const;
 
-	TSharedRef<SWidget> GetRowExtensionButtons(TSharedPtr<IPropertyHandle> InPropertyHandle) const;
+	TSharedRef<SWidget> GetRowExtensionButtons(TSharedPtr<IPropertyHandle> InPropertyHandle);
 private:
 
 	/** The node info to build the tree view row from. */
@@ -197,6 +197,12 @@ private:
 
 	/** The set of material parameters this is associated with */
 	TWeakObjectPtr<UMaterialInstance> MaterialInstance;
+
+	/** Pointer to copied ValuePtr. Needed for workaround for now */
+	TUniquePtr<uint8> NewDefaultValuePtr;
+
+	/** Pointer to Widget for Workaround. */
+	TSharedPtr<SWidget> ResetArrow;
 };
 
 
@@ -216,20 +222,42 @@ FSlateColor SMaterialDynamicParametersOverviewTreeItem::GetOuterBackgroundColor(
 	return FAppStyle::Get().GetSlateColor("Colors.Panel");
 }
 
-TSharedRef<SWidget> SMaterialDynamicParametersOverviewTreeItem::GetRowExtensionButtons(TSharedPtr<IPropertyHandle> InPropertyHandle) const
+TSharedRef<SWidget> SMaterialDynamicParametersOverviewTreeItem::GetRowExtensionButtons(TSharedPtr<IPropertyHandle> InPropertyHandle)
 {
 	if (!InPropertyHandle.IsValid())
 	{
 		return SNullWidget::NullWidget;
 	}
 
-	const auto ResetDelegate = FExecuteAction::CreateLambda([InPropertyHandle]()
+	const auto ResetDelegate = FExecuteAction::CreateLambda([this, InPropertyHandle]()
 	{
 		if (InPropertyHandle.IsValid())
 		{
 			InPropertyHandle->ResetToDefault();
+			ResetArrow->SetVisibility(EVisibility::Hidden);
 		}
 	});
+
+	/** Needed as Dynamic Material Instances don't have a default to use normally for DiffersFromDefault */
+	const auto DiffersFromDefaultDelegate = FSimpleDelegate::CreateLambda([this, InPropertyHandle]()
+	{
+		if (!NewDefaultValuePtr)
+		{
+			return;
+		}
+
+		FProperty* Property = InPropertyHandle->GetProperty();
+		
+		void* ValuePtr;
+		InPropertyHandle->GetValueData(ValuePtr);
+		
+		if (Property->Identical(NewDefaultValuePtr.Get(), ValuePtr))
+		{
+			ResetArrow->SetVisibility(EVisibility::Hidden);
+			return;
+		}
+		ResetArrow->SetVisibility(EVisibility::Visible);
+	});
 	
 	FPropertyEditorModule& PropertyEditorModule = FModuleManager::Get().GetModuleChecked<FPropertyEditorModule>("PropertyEditor");
 
@@ -255,7 +283,20 @@ TSharedRef<SWidget> SMaterialDynamicParametersOverviewTreeItem::GetRowExtensionB
 		ToolbarBuilder.AddToolBarButton(Extension.UIAction, NAME_None, Extension.Label, Extension.ToolTip, Extension.Icon);
 	}
 
-	return ToolbarBuilder.MakeWidget();
+	/** Set up the Workaround. Initially hidden since it's set as the default. */
+	FProperty* Property = InPropertyHandle->GetProperty();
+	NewDefaultValuePtr = MakeUnique<uint8>(Property->GetSize());
+	void* ValuePtr;
+	InPropertyHandle->GetValueData(ValuePtr);
+	Property->CopyCompleteValue(NewDefaultValuePtr.Get(), ValuePtr);
+	Property->ClearValue(NewDefaultValuePtr.Get());
+	
+	InPropertyHandle->SetOnPropertyValueChanged(DiffersFromDefaultDelegate);
+	InPropertyHandle->SetOnChildPropertyValueChanged(DiffersFromDefaultDelegate);
+
+	ResetArrow = ToolbarBuilder.MakeWidget();
+	
+	return ResetArrow.ToSharedRef();
 }
 
 void SMaterialDynamicParametersOverviewTreeItem::RefreshOnRowChange(const FAssetData& AssetData, TSharedPtr<SMaterialDynamicParametersOverviewTree> InTree)
