VPUtilities: Wrapped a call to composure layer RerunConstructionScripts in WITH_EDITOR to allow for packaging.

diff --git a/Engine/Plugins/Experimental/VirtualProductionUtilities/Source/VPUtilities/Private/VPFullScreenUserWidget.cpp b/Engine/Plugins/Experimental/VirtualProductionUtilities/Source/VPUtilities/Private/VPFullScreenUserWidget.cpp
--- a/Engine/Plugins/Experimental/VirtualProductionUtilities/Source/VPUtilities/Private/VPFullScreenUserWidget.cpp
+++ b/Engine/Plugins/Experimental/VirtualProductionUtilities/Source/VPUtilities/Private/VPFullScreenUserWidget.cpp
@@ -735,7 +735,9 @@ bool UVPFullScreenUserWidget::Display(UWorld* InWorld)
 								if (DestTextureRT2D)
 								{
 									TextureInputObjectProperty->SetObjectPropertyValue(DestTextureRT2D, PostProcessDisplayType.WidgetRenderTarget);
+#if WITH_EDITOR
 									Layer->RerunConstructionScripts();
+#endif // WITH_EDITOR
 								}
 							}
 						}
