Fix ResetArrow Visibility Bug.
- Seperated Visibility from the Toolbar Widget, else issue with the expose button.

diff --git a/Engine/Plugins/Experimental/VirtualProductionUtilities/Source/VPMaterialsEditor/Private/SMaterialDynamicParametersPanelWidget.cpp b/Engine/Plugins/Experimental/VirtualProductionUtilities/Source/VPMaterialsEditor/Private/SMaterialDynamicParametersPanelWidget.cpp
--- a/Engine/Plugins/Experimental/VirtualProductionUtilities/Source/VPMaterialsEditor/Private/SMaterialDynamicParametersPanelWidget.cpp
+++ b/Engine/Plugins/Experimental/VirtualProductionUtilities/Source/VPMaterialsEditor/Private/SMaterialDynamicParametersPanelWidget.cpp
@@ -239,11 +239,11 @@ TSharedRef<SWidget> SMaterialDynamicParametersOverviewTreeItem::GetRowExtensionB
 	});
 
 	/** Needed as Dynamic Material Instances don't have a default to use normally for DiffersFromDefault */
-	const auto DiffersFromDefaultDelegate = FSimpleDelegate::CreateLambda([this, InPropertyHandle]()
+	const auto DiffersFromDefaultDelegate = FIsActionButtonVisible::CreateLambda([this, InPropertyHandle]()
 	{
 		if (!NewDefaultValuePtr)
 		{
-			return;
+			return true;
 		}
 
 		FProperty* Property = InPropertyHandle->GetProperty();
@@ -253,10 +253,9 @@ TSharedRef<SWidget> SMaterialDynamicParametersOverviewTreeItem::GetRowExtensionB
 		
 		if (Property->Identical(NewDefaultValuePtr.Get(), ValuePtr))
 		{
-			ResetArrow->SetVisibility(EVisibility::Hidden);
-			return;
+			return false;
 		}
-		ResetArrow->SetVisibility(EVisibility::Visible);
+		return true;
 	});
 	
 	FPropertyEditorModule& PropertyEditorModule = FModuleManager::Get().GetModuleChecked<FPropertyEditorModule>("PropertyEditor");
@@ -264,7 +263,7 @@ TSharedRef<SWidget> SMaterialDynamicParametersOverviewTreeItem::GetRowExtensionB
 	TArray<FPropertyRowExtensionButton> ExtensionButtons;
 	FPropertyRowExtensionButton& ResetToDefault = ExtensionButtons.AddDefaulted_GetRef();
 	ResetToDefault.Label = NSLOCTEXT("PropertyEditor", "ResetToDefault", "Reset to Default");
-	ResetToDefault.UIAction = FUIAction(ResetDelegate);
+	ResetToDefault.UIAction = FUIAction(ResetDelegate, FCanExecuteAction(), FIsActionChecked(), DiffersFromDefaultDelegate);
 
 	ResetToDefault.Icon = FSlateIcon(FAppStyle::Get().GetStyleSetName(), "PropertyWindow.DiffersFromDefault");
 	ResetToDefault.ToolTip = NSLOCTEXT("PropertyEditor", "ResetToDefaultToolTip", "Reset this property to its default value.");
@@ -283,7 +282,7 @@ TSharedRef<SWidget> SMaterialDynamicParametersOverviewTreeItem::GetRowExtensionB
 		ToolbarBuilder.AddToolBarButton(Extension.UIAction, NAME_None, Extension.Label, Extension.ToolTip, Extension.Icon);
 	}
 
-	/** Set up the Workaround. Initially hidden since it's set as the default. */
+	/** Set up the Default Value for the Workaround. */
 	FProperty* Property = InPropertyHandle->GetProperty();
 	NewDefaultValuePtr = MakeUnique<uint8>(Property->GetSize());
 	void* ValuePtr;
@@ -291,9 +290,6 @@ TSharedRef<SWidget> SMaterialDynamicParametersOverviewTreeItem::GetRowExtensionB
 	Property->CopyCompleteValue(NewDefaultValuePtr.Get(), ValuePtr);
 	Property->ClearValue(NewDefaultValuePtr.Get());
 	
-	InPropertyHandle->SetOnPropertyValueChanged(DiffersFromDefaultDelegate);
-	InPropertyHandle->SetOnChildPropertyValueChanged(DiffersFromDefaultDelegate);
-
 	ResetArrow = ToolbarBuilder.MakeWidget();
 	
 	return ResetArrow.ToSharedRef();
