MU bug. Dynamic material instances do not transact to nDisplay until Save All

diff --git a/Engine/Plugins/Experimental/VirtualProductionUtilities/Source/VPMaterialsEditor/Private/SMaterialDynamicWidgets.cpp b/Engine/Plugins/Experimental/VirtualProductionUtilities/Source/VPMaterialsEditor/Private/SMaterialDynamicWidgets.cpp
--- a/Engine/Plugins/Experimental/VirtualProductionUtilities/Source/VPMaterialsEditor/Private/SMaterialDynamicWidgets.cpp
+++ b/Engine/Plugins/Experimental/VirtualProductionUtilities/Source/VPMaterialsEditor/Private/SMaterialDynamicWidgets.cpp
@@ -105,6 +105,12 @@ void SMaterialDynamicView::Construct(const FArguments& InArgs, const TSharedRef<
 
 FReply SMaterialDynamicView::OnResetButtonClicked() const
 {
+	UActorComponent* OwnerActorComponent = CurrentComponent.Get();
+	if (!ensure(OwnerActorComponent))
+	{
+		return FReply::Handled();
+	}
+
 	const TSharedPtr<FMaterialItemView> MaterialItemView = MaterialItemViewWeakPtr.Pin();
 	if(!ensure(MaterialItemView.IsValid()))
 	{
@@ -123,6 +129,10 @@ FReply SMaterialDynamicView::OnResetButtonClicked() const
 		return FReply::Handled();
 	}
 
+
+	FScopedTransaction Transaction(LOCTEXT("ResetDynamicMaterial", "Reset Dynamic Material"));
+	MaterialInstanceDynamic->Modify();
+	OwnerActorComponent->Modify();
 	MaterialInstanceDynamic->CopyParameterOverrides(ParentMaterialInstance);
 
 	return FReply::Handled();
@@ -130,6 +140,12 @@ FReply SMaterialDynamicView::OnResetButtonClicked() const
 
 FReply SMaterialDynamicView::OnRevertButtonClicked() const
 {
+	UActorComponent* OwnerActorComponent = CurrentComponent.Get();
+	if (!ensure(OwnerActorComponent))
+	{
+		return FReply::Handled();
+	}
+
 	const TSharedPtr<FMaterialItemView> MaterialItemView = MaterialItemViewWeakPtr.Pin();
 	if(!ensure(MaterialItemView.IsValid()))
 	{
@@ -142,6 +158,7 @@ FReply SMaterialDynamicView::OnRevertButtonClicked() const
 		return FReply::Handled();
 	}
 
+	FScopedTransaction Transaction(LOCTEXT("RevertDynamicMaterial", "Revert Dynamic Material"));
 	MaterialItemView->ReplaceMaterial(ParentMaterialInstance);
 	
 	return FReply::Handled();
@@ -149,6 +166,12 @@ FReply SMaterialDynamicView::OnRevertButtonClicked() const
 
 FReply SMaterialDynamicView::OnCopyToOriginalButtonClicked() const
 {
+	UActorComponent* OwnerActorComponent = CurrentComponent.Get();
+	if (!ensure(OwnerActorComponent))
+	{
+		return FReply::Handled();
+	}
+
 	const TSharedPtr<FMaterialItemView> MaterialItemView = MaterialItemViewWeakPtr.Pin();
 	if(!ensure(MaterialItemView.IsValid()))
 	{
@@ -167,6 +190,7 @@ FReply SMaterialDynamicView::OnCopyToOriginalButtonClicked() const
 		return FReply::Handled();
 	}
 
+	FScopedTransaction Transaction(LOCTEXT("CopyToOriginalMaterial", "Copy To Original Material"));
 	ParentMaterialInstance->ScalarParameterValues = MaterialInstanceDynamic->ScalarParameterValues;
 	ParentMaterialInstance->VectorParameterValues = MaterialInstanceDynamic->VectorParameterValues;
 	ParentMaterialInstance->TextureParameterValues = MaterialInstanceDynamic->TextureParameterValues;
@@ -196,12 +220,17 @@ FReply SMaterialDynamicView::OnCreateDynamicMaterialButtonClicked() const
 
 	if (MaterialInstance)
 	{
-		// Create Material
-		UMaterialInstanceDynamic* NewMaterialInstanceDynamic = UMaterialInstanceDynamic::Create(MaterialInstance, OwnerActorComponent);
+		OwnerActorComponent->Modify();
 
-		// Set object Transactional in order Undo, Redo working correctly and have MU support
-		NewMaterialInstanceDynamic->SetFlags(RF_Transactional);
-		
+		FScopedTransaction Transaction(LOCTEXT("CreateDynamicMaterial", "Create Dynamic Material"));
+
+		// Create Transactional Material. That is allawed create MID with MU and transact the changes
+		UMaterialInstanceDynamic* NewMaterialInstanceDynamic = NewObject<UMaterialInstanceDynamic>(OwnerActorComponent, NAME_None, RF_Transactional);
+		NewMaterialInstanceDynamic->Parent = MaterialInstance;
+		NewMaterialInstanceDynamic->UpdateCachedData();
+
+
+		// Copy Parameter Overrides with Transaction
 		NewMaterialInstanceDynamic->CopyParameterOverrides(MaterialInstance);
 		MaterialItemView->ReplaceMaterial(NewMaterialInstanceDynamic);
 	}
