NoPCH Compile LyraEditor Win64 - 'FScopedTransaction': undeclared identifier in SMaterialDynamicWidgets.cpp
#rnx

diff --git a/Engine/Plugins/Experimental/VirtualProductionUtilities/Source/VPMaterialsEditor/Private/SMaterialDynamicWidgets.cpp b/Engine/Plugins/Experimental/VirtualProductionUtilities/Source/VPMaterialsEditor/Private/SMaterialDynamicWidgets.cpp
--- a/Engine/Plugins/Experimental/VirtualProductionUtilities/Source/VPMaterialsEditor/Private/SMaterialDynamicWidgets.cpp
+++ b/Engine/Plugins/Experimental/VirtualProductionUtilities/Source/VPMaterialsEditor/Private/SMaterialDynamicWidgets.cpp
@@ -3,6 +3,7 @@
 #include "SMaterialDynamicWidgets.h"
 
 #include "PropertyCustomizationHelpers.h"
+#include "ScopedTransaction.h"
 #include "SMaterialDynamicParametersPanelWidget.h"
 #include "Components/ActorComponent.h"
 #include "Materials/MaterialInstanceDynamic.h"
