- Fixing crash when loading cooked QAVirtualproduction project by putting back vputilities as a runtime module.

diff --git a/Engine/Plugins/Experimental/VirtualProductionUtilities/VirtualProductionUtilities.uplugin b/Engine/Plugins/Experimental/VirtualProductionUtilities/VirtualProductionUtilities.uplugin
--- a/Engine/Plugins/Experimental/VirtualProductionUtilities/VirtualProductionUtilities.uplugin
+++ b/Engine/Plugins/Experimental/VirtualProductionUtilities/VirtualProductionUtilities.uplugin
@@ -16,7 +16,7 @@
 	"Modules": [
 		{
 			"Name": "VPUtilities",
-			"Type": "UncookedOnly",
+			"Type": "Runtime",
 			"LoadingPhase": "PostDefault"
 		},
 		{
