VPUtils: Removed unneeded dependency on CCRs.

diff --git a/Engine/Plugins/Experimental/VirtualProductionUtilities/VirtualProductionUtilities.uplugin b/Engine/Plugins/Experimental/VirtualProductionUtilities/VirtualProductionUtilities.uplugin
--- a/Engine/Plugins/Experimental/VirtualProductionUtilities/VirtualProductionUtilities.uplugin
+++ b/Engine/Plugins/Experimental/VirtualProductionUtilities/VirtualProductionUtilities.uplugin
@@ -64,10 +64,6 @@
 		{
 			"Name": "Composure",
 			"Enabled": true
-		},
-		{
-			"Name": "ColorCorrectRegions",
-			"Enabled": true
 		}
 	]
 }
