Virtual Scouting: Add a pseudo-"game view" toggle.

This isn't true game view, because the existing game view code path results in hiding the VR interactors, among other things. Instead, we clear the editor showflag, which is sufficient to hide actors marked "hidden in game" (among other things). To keep scouting tools visible, some adjustments needed to be made to tool blueprint "hidden in game" properties.

This change also adds an "Enter VR Mode" BP method, which has been requested before (UE-139643) and complements the existing "Exit VR Mode" method.

diff --git a/Engine/Plugins/Experimental/VirtualProductionUtilities/Source/VPUtilitiesEditor/Private/VPScoutingSubsystem.cpp b/Engine/Plugins/Experimental/VirtualProductionUtilities/Source/VPUtilitiesEditor/Private/VPScoutingSubsystem.cpp
--- a/Engine/Plugins/Experimental/VirtualProductionUtilities/Source/VPUtilitiesEditor/Private/VPScoutingSubsystem.cpp
+++ b/Engine/Plugins/Experimental/VirtualProductionUtilities/Source/VPUtilitiesEditor/Private/VPScoutingSubsystem.cpp
@@ -372,6 +372,23 @@ void UVPScoutingSubsystem::SetIsHelperSystemEnabled(const bool bInIsHelperSystem
 	VPUtilitiesEditorSettings->SaveConfig();
 }
 
+UVREditorMode* UVPScoutingSubsystem::GetVREditorMode()
+{
+	return IVREditorModule::Get().GetVRMode();
+}
+
+bool UVPScoutingSubsystem::EnterVRMode()
+{
+	IVREditorModule& VREditorModule = IVREditorModule::Get();
+	if (VREditorModule.IsVREditorAvailable())
+	{
+		VREditorModule.EnableVREditor(true);
+		return true;
+	}
+
+	return false;
+}
+
 void UVPScoutingSubsystem::ExitVRMode()
 {
 	IVREditorModule& VREditorModule = IVREditorModule::Get();
diff --git a/Engine/Plugins/Experimental/VirtualProductionUtilities/Source/VPUtilitiesEditor/Private/VRTool.cpp b/Engine/Plugins/Experimental/VirtualProductionUtilities/Source/VPUtilitiesEditor/Private/VRTool.cpp
--- a/Engine/Plugins/Experimental/VirtualProductionUtilities/Source/VPUtilitiesEditor/Private/VRTool.cpp
+++ b/Engine/Plugins/Experimental/VirtualProductionUtilities/Source/VPUtilitiesEditor/Private/VRTool.cpp
@@ -5,6 +5,12 @@
 #include "Engine/InputDelegateBinding.h"
 
 
+AVRTool::AVRTool(const FObjectInitializer& ObjectInitializer)
+	: Super(ObjectInitializer)
+{
+	// We don't want VR "game view" to hide the tools.
+	SetActorHiddenInGame(false);
+}
 
 void AVRTool::OnConstruction(const FTransform& Transform)
 {
diff --git a/Engine/Plugins/Experimental/VirtualProductionUtilities/Source/VPUtilitiesEditor/Public/VPScoutingSubsystem.h b/Engine/Plugins/Experimental/VirtualProductionUtilities/Source/VPUtilitiesEditor/Public/VPScoutingSubsystem.h
--- a/Engine/Plugins/Experimental/VirtualProductionUtilities/Source/VPUtilitiesEditor/Public/VPScoutingSubsystem.h
+++ b/Engine/Plugins/Experimental/VirtualProductionUtilities/Source/VPUtilitiesEditor/Public/VPScoutingSubsystem.h
@@ -210,7 +210,15 @@ public:
 	UFUNCTION(BlueprintCallable, Category = "Virtual Production")
 	static void SetIsHelperSystemEnabled(const bool bInIsHelperSystemEnabled);
 
-	/** Exit VR Mode  */
+	/** Get VR Editor Mode object */
+	UFUNCTION(BlueprintPure, Category = "Virtual Production")
+	static UVREditorMode* GetVREditorMode();
+
+	/** Enter VR Mode */
+	UFUNCTION(BlueprintCallable, Category = "Virtual Production")
+	static bool EnterVRMode();
+
+	/** Exit VR Mode */
 	UFUNCTION(BlueprintCallable, Category = "Virtual Production")
 	static void ExitVRMode();
 
@@ -224,11 +232,11 @@ public:
 
 	/** Whether rotation grid snapping is enabled */
 	UFUNCTION(BlueprintCallable, Category = "Virtual Production")
-		static bool IsRotationGridSnappingEnabled();
+	static bool IsRotationGridSnappingEnabled();
 
 	/** Toggle rotation grid snapping */
 	UFUNCTION(BlueprintCallable, Category = "Virtual Production")
-		static void ToggleRotationGridSnapping();
+	static void ToggleRotationGridSnapping();
 
 private:
 
diff --git a/Engine/Plugins/Experimental/VirtualProductionUtilities/Source/VPUtilitiesEditor/Public/VRTool.h b/Engine/Plugins/Experimental/VirtualProductionUtilities/Source/VPUtilitiesEditor/Public/VRTool.h
--- a/Engine/Plugins/Experimental/VirtualProductionUtilities/Source/VPUtilitiesEditor/Public/VRTool.h
+++ b/Engine/Plugins/Experimental/VirtualProductionUtilities/Source/VPUtilitiesEditor/Public/VRTool.h
@@ -14,6 +14,8 @@ class AVRTool : public AVPTransientEditorTickableActorBase
 {
 	GENERATED_BODY()
 
+	AVRTool(const FObjectInitializer& ObjectInitializer);
+
 	virtual void OnConstruction(const FTransform& Transform) override;
 #if WITH_EDITOR
 	virtual void PostEditChangeProperty(FPropertyChangedEvent& PropertyChangedEvent) override;
diff --git a/Engine/Source/Editor/VREditor/Public/VREditorMode.h b/Engine/Source/Editor/VREditor/Public/VREditorMode.h
--- a/Engine/Source/Editor/VREditor/Public/VREditorMode.h
+++ b/Engine/Source/Editor/VREditor/Public/VREditorMode.h
@@ -207,6 +207,14 @@ public:
 	/** Mutable version of above. */
 	class SLevelViewport& GetLevelViewportPossessedForVR();
 
+	/** Display the scene more closely to how it would appear at runtime (as opposed to edit time). */
+	UFUNCTION(BlueprintCallable, Category="VREditorMode")
+	void SetGameView(bool bGameView);
+
+	/** Returns whether game view is currently active. */
+	UFUNCTION(BlueprintCallable, Category="VREditorMode")
+	bool IsInGameView() const;
+
 	/** Gets the world scale factor, which can be multiplied by a scale vector to convert to room space */
 	UFUNCTION( BlueprintCallable, Category = "VREditorMode" )
 	float GetWorldScaleFactor() const;
diff --git a/Engine/Source/Editor/VREditor/VREditorMode.cpp b/Engine/Source/Editor/VREditor/VREditorMode.cpp
--- a/Engine/Source/Editor/VREditor/VREditorMode.cpp
+++ b/Engine/Source/Editor/VREditor/VREditorMode.cpp
@@ -693,6 +693,28 @@ SLevelViewport& UVREditorMode::GetLevelViewportPossessedForVR()
 	return *VREditorLevelViewportWeakPtr.Pin();
 }
 
+void UVREditorMode::SetGameView(bool bGameView)
+{
+	if (TSharedPtr<SLevelViewport> Viewport = VREditorLevelViewportWeakPtr.Pin())
+	{
+		// We can't actually set the viewport to game view, because AVREditorAvatarActor::IsEditorOnly is
+		// overridden to return true, and so its owned components (including the interactors) get hidden.
+		// However, clearing the "editor" flag turns out to get us close to what we'd want.
+		FLevelEditorViewportClient& ViewportClient = Viewport->GetLevelViewportClient();
+		ViewportClient.EngineShowFlags.SetEditor(!bGameView);
+	}
+}
+
+bool UVREditorMode::IsInGameView() const
+{
+	if (TSharedPtr<SLevelViewport> Viewport = VREditorLevelViewportWeakPtr.Pin())
+	{
+		FLevelEditorViewportClient& ViewportClient = Viewport->GetLevelViewportClient();
+		return ViewportClient.EngineShowFlags.Editor == 0;
+	}
+
+	return false;
+}
 
 float UVREditorMode::GetWorldScaleFactor() const
 {
