VP Roles improvements
Move roles to a menu in the level editor toolbar
Make it easier to select/remove roles
Move roles to its own module
Add roles subsystem to do basic operations on roles at runtime
Move enabled roles to a user setting

diff --git a/Engine/Plugins/Editor/GameplayTagsEditor/Source/GameplayTagsEditor/Private/SGameplayTagWidget.cpp b/Engine/Plugins/Editor/GameplayTagsEditor/Source/GameplayTagsEditor/Private/SGameplayTagWidget.cpp
--- a/Engine/Plugins/Editor/GameplayTagsEditor/Source/GameplayTagsEditor/Private/SGameplayTagWidget.cpp
+++ b/Engine/Plugins/Editor/GameplayTagsEditor/Source/GameplayTagsEditor/Private/SGameplayTagWidget.cpp
@@ -159,6 +159,7 @@ void SGameplayTagWidget::Construct(const FArguments& InArgs, const TArray<FEdita
 	ChildSlot
 	[
 		SNew(SBorder)
+		.Padding(InArgs._Padding)
 		.BorderImage(InArgs._BackgroundBrush)
 		[
 			SNew(SVerticalBox)
@@ -304,9 +305,7 @@ void SGameplayTagWidget::Construct(const FArguments& InArgs, const TArray<FEdita
 			+SVerticalBox::Slot()
 			.MaxHeight(MaxHeight)
 			[
-				SAssignNew(TagTreeContainerWidget, SBorder)
-				.BorderImage(FAppStyle::GetBrush("NoBorder"))
-				.Padding(FMargin(4.f))
+				SAssignNew(TagTreeContainerWidget, SBox)
 				[
 					SAssignNew(TagTreeWidget, STreeView< TSharedPtr<FGameplayTagNode> >)
 					.TreeItemsSource(&TagItems)
diff --git a/Engine/Plugins/Editor/GameplayTagsEditor/Source/GameplayTagsEditor/Public/SGameplayTagWidget.h b/Engine/Plugins/Editor/GameplayTagsEditor/Source/GameplayTagsEditor/Public/SGameplayTagWidget.h
--- a/Engine/Plugins/Editor/GameplayTagsEditor/Source/GameplayTagsEditor/Public/SGameplayTagWidget.h
+++ b/Engine/Plugins/Editor/GameplayTagsEditor/Source/GameplayTagsEditor/Public/SGameplayTagWidget.h
@@ -44,7 +44,8 @@ public:
 	DECLARE_DELEGATE_RetVal_OneParam( ETagFilterResult, FOnFilterTag, const TSharedPtr<FGameplayTagNode>&)
 
 	SLATE_BEGIN_ARGS( SGameplayTagWidget )
-		: _Filter()
+		: _Padding(2.f)
+		, _Filter()
 		, _NewTagName(TEXT(""))
 		, _ReadOnly( false )
 		, _TagContainerName( TEXT("") )
@@ -61,6 +62,7 @@ public:
 		, _TagTreeViewBackgroundBrush(nullptr)
 		
 	{}
+		SLATE_ATTRIBUTE( FMargin, Padding ) // Padding for the containing border.
 		SLATE_ARGUMENT( FString, Filter ) // Comma delimited string of tag root names to filter by
 		SLATE_EVENT( FOnFilterTag, OnFilterTag ) // Optional filter function called when generating the tag list
 		SLATE_ARGUMENT( FString, NewTagName ) // String that will initially populate the New Tag Name field
@@ -197,7 +199,7 @@ private:
 	TArray< TSharedPtr<FGameplayTagNode> > FilteredTagItems;
 
 	/** Container widget holding the tag tree */
-	TSharedPtr<SBorder> TagTreeContainerWidget;
+	TSharedPtr<SWidget> TagTreeContainerWidget;
 
 	/** Tree widget showing the gameplay tag library */
 	TSharedPtr< STreeView< TSharedPtr<FGameplayTagNode> > > TagTreeWidget;
diff --git a/Engine/Plugins/Experimental/VirtualCamera/Source/VCamCore/Private/VCamComponent.cpp b/Engine/Plugins/Experimental/VirtualCamera/Source/VCamCore/Private/VCamComponent.cpp
--- a/Engine/Plugins/Experimental/VirtualCamera/Source/VCamCore/Private/VCamComponent.cpp
+++ b/Engine/Plugins/Experimental/VirtualCamera/Source/VCamCore/Private/VCamComponent.cpp
@@ -32,6 +32,7 @@
 #include "EnhancedInputEditorSubsystem.h"
 
 #include "VPSettings.h"
+#include "VPRolesSubsystem.h"
 #endif
 
 
@@ -1491,7 +1492,7 @@ bool UVCamComponent::IsCameraInVPRole() const
 	// We are in a valid camera role if the user has not assigned a role or the current VPSettings role matches the
 	// assigned role.
 	//
-	return !Role.IsValid() || Settings->GetRoles().HasTag(Role);
+	return !Role.IsValid() || GEngine->GetEngineSubsystem<UVirtualProductionRolesSubsystem>()->HasActiveRole(Role);
 #else
 	return true;
 #endif
diff --git a/Engine/Plugins/Experimental/VirtualCamera/Source/VCamCore/VCamCore.Build.cs b/Engine/Plugins/Experimental/VirtualCamera/Source/VCamCore/VCamCore.Build.cs
--- a/Engine/Plugins/Experimental/VirtualCamera/Source/VCamCore/VCamCore.Build.cs
+++ b/Engine/Plugins/Experimental/VirtualCamera/Source/VCamCore/VCamCore.Build.cs
@@ -47,7 +47,8 @@ public class VCamCore : ModuleRules
 				"LiveLinkInterface",
 				"MediaIOCore",
 				"RemoteSession",
-				"InputCore"
+				"InputCore",
+				"VPRoles",
 				// ... add private dependencies that you statically link with here ...
 			}
 			);
diff --git a/Engine/Plugins/Experimental/VirtualProductionUtilities/Source/VPRoles/Private/VPRolesModule.cpp b/Engine/Plugins/Experimental/VirtualProductionUtilities/Source/VPRoles/Private/VPRolesModule.cpp
new file mode 100644
--- /dev/null
+++ b/Engine/Plugins/Experimental/VirtualProductionUtilities/Source/VPRoles/Private/VPRolesModule.cpp
@@ -0,0 +1,10 @@
+// Copyright Epic Games, Inc. All Rights Reserved.
+
+#include "VPRolesModule.h"
+#include "CoreMinimal.h"
+#include "Modules/ModuleManager.h"
+#include "Modules/ModuleInterface.h"
+
+DEFINE_LOG_CATEGORY(LogVPRoles);
+
+IMPLEMENT_MODULE(FDefaultModuleImpl, VPRoles)
diff --git a/Engine/Plugins/Experimental/VirtualProductionUtilities/Source/VPRoles/Private/VPRolesModule.h b/Engine/Plugins/Experimental/VirtualProductionUtilities/Source/VPRoles/Private/VPRolesModule.h
new file mode 100644
--- /dev/null
+++ b/Engine/Plugins/Experimental/VirtualProductionUtilities/Source/VPRoles/Private/VPRolesModule.h
@@ -0,0 +1,7 @@
+// Copyright Epic Games, Inc. All Rights Reserved.
+
+#pragma once
+
+#include "Logging/LogMacros.h"
+
+DECLARE_LOG_CATEGORY_EXTERN(LogVPRoles, Log, Log);
diff --git a/Engine/Plugins/Experimental/VirtualProductionUtilities/Source/VPRoles/Private/VPRolesSettings.h b/Engine/Plugins/Experimental/VirtualProductionUtilities/Source/VPRoles/Private/VPRolesSettings.h
new file mode 100644
--- /dev/null
+++ b/Engine/Plugins/Experimental/VirtualProductionUtilities/Source/VPRoles/Private/VPRolesSettings.h
@@ -0,0 +1,32 @@
+// Copyright Epic Games, Inc. All Rights Reserved.
+
+#pragma once
+
+#include "CoreMinimal.h"
+#include "GameplayTagContainer.h"
+#include "VPRolesSettings.generated.h"
+
+UCLASS(config=UserVPRoles)
+class VPROLES_API UVPRolesUserSettings : public UObject
+{
+public:
+	GENERATED_BODY()
+
+private:
+	/**
+	 * The machine role(s) in a virtual production context.
+	 * @note The role may be override via the command line, "-VPRole=[Role.SubRole1|Role.SubRole2]"
+	 */
+	UPROPERTY(config)
+	FGameplayTagContainer Roles;
+	
+	/** The files that contain the available VP Roles. */
+	UPROPERTY(config)
+	TSet<FName> RoleSources;
+	
+	/** Cached available roles since the list isn't available outside of editor. */
+	UPROPERTY(config)
+	TArray<FString> CachedRoles;
+	
+	friend class UVirtualProductionRolesSubsystem;
+};
diff --git a/Engine/Plugins/Experimental/VirtualProductionUtilities/Source/VPRoles/Private/VPRolesSubsystem.cpp b/Engine/Plugins/Experimental/VirtualProductionUtilities/Source/VPRoles/Private/VPRolesSubsystem.cpp
new file mode 100644
--- /dev/null
+++ b/Engine/Plugins/Experimental/VirtualProductionUtilities/Source/VPRoles/Private/VPRolesSubsystem.cpp
@@ -0,0 +1,280 @@
+// Copyright Epic Games, Inc. All Rights Reserved.
+
+#include "VPRolesSubsystem.h"
+
+#include "Algo/Transform.h"
+#include "GameplayTagContainer.h"
+#include "GameplayTagsManager.h"
+#include "Misc/CommandLine.h"
+#include "VPRolesSettings.h"
+#include "VPRolesModule.h"
+#include "VPSettings.h"
+
+#if WITH_EDITOR
+#include "Editor.h"
+#include "GameplayTagsEditorModule.h"
+#endif
+
+static const FName DefaultVPRolesFileName = ("VPRoles.ini");
+
+void UVirtualProductionRolesSubsystem::Initialize(FSubsystemCollectionBase& Collection)
+{
+	CommandLineRoles = MakePimpl<FGameplayTagContainer>();
+	ImportCurrentRolesFromCommandLine();
+	ImportCurrentRolesFromDeprecatedSetting();
+
+	AddRoleSource(DefaultVPRolesFileName);
+
+	// todo: Pop up menu to migrate all sources to the VPRoles one.
+	ImportRoleSourcesFromVPSettings();
+	
+	// Provide a getter to the old system so it can return up to date roles.
+	GetMutableDefault<UVPSettings>()->RolesGetter.BindUObject(this, &UVirtualProductionRolesSubsystem::GetRolesContainer_Private);
+}
+
+void UVirtualProductionRolesSubsystem::Deinitialize()
+{
+	if (UVPSettings* VPSettings = GetMutableDefault<UVPSettings>())
+	{
+		// If this is uninitialized before the VP Settings are destroyed, provide a fallback that returns an empty container.
+		VPSettings->RolesGetter.BindStatic(&UVirtualProductionRolesSubsystem::GetRolesContainerFallback);
+	}
+}
+
+TArray<FString> UVirtualProductionRolesSubsystem::GetActiveRoles() const
+{
+	TArray<FString> Roles;
+	Algo::Transform(GetRolesContainer_Private(), Roles, [](const FGameplayTag& Tag) { return Tag.ToString(); });
+	return Roles;
+}
+
+bool UVirtualProductionRolesSubsystem::HasActiveRole(const FString& Role) const
+{
+	constexpr bool bErrorWhenNotFound = false;
+	const FGameplayTag Tag = UGameplayTagsManager::Get().RequestGameplayTag(*Role, bErrorWhenNotFound);
+	if (Tag.IsValid())
+	{
+		return GetRolesContainer_Private().HasTag(Tag);
+	}
+	return false;
+}
+
+FString UVirtualProductionRolesSubsystem::GetActiveRolesString() const
+{
+	return GetRolesContainer_Private().ToStringSimple();
+}
+
+void UVirtualProductionRolesSubsystem::SetActiveRoles(const TArray<FString>& Roles)
+{
+	FGameplayTagContainer Container;
+	
+	for (const FString& Role : Roles)
+	{
+		constexpr bool bErrorIfNotFound = false;
+		FGameplayTag Tag = FGameplayTag::RequestGameplayTag(*Role, bErrorIfNotFound);
+		if (Tag.IsValid())
+		{
+			Container.AddTag(Tag);
+		}
+		else
+		{
+			UE_LOG(LogVPRoles, Error, TEXT("Role %s doesn't exist."), *Role);
+		}
+	}
+
+	// Clear the transient roles if we're running through code while in editor.
+	GetMutableDefault<UVPRolesUserSettings>()->Roles = MoveTemp(Container);
+
+	BroadcastRolesChanged();
+}
+
+TArray<FString> UVirtualProductionRolesSubsystem::GetAllRoles() const
+{
+#if WITH_EDITOR
+	// Get all roles from our collected sources.
+	TArray<TSharedPtr<FGameplayTagNode>> Roles;
+	for (FName RoleSource : GetRoleSources())
+	{
+		UGameplayTagsManager::Get().GetAllTagsFromSource(RoleSource, Roles);
+	}
+
+	TArray<FString> ReturnedRoles;
+	ReturnedRoles.Reserve(Roles.Num());
+	Algo::TransformIf(Roles, ReturnedRoles,
+		[](const TSharedPtr<FGameplayTagNode>& Node){ return Node.IsValid(); },
+		[](const TSharedPtr<FGameplayTagNode>& Node){ return Node->GetCompleteTagString(); });
+	
+	return ReturnedRoles;
+#else
+	return GetDefault<UVPRolesUserSettings>()->CachedRoles;
+#endif
+}
+
+#if WITH_EDITOR
+bool UVirtualProductionRolesSubsystem::AddRole(const FString& RoleName)
+{
+	const bool bResult = IGameplayTagsEditorModule::Get().AddNewGameplayTagToINI(RoleName, TEXT("VPRole"), DefaultVPRolesFileName);
+	BroadcastRolesChanged();
+	return bResult;
+}
+	
+bool UVirtualProductionRolesSubsystem::RemoveRole(const FString& RoleName)
+{
+	bool bResult = false;
+	if (const TSharedPtr<FGameplayTagNode> Node = UGameplayTagsManager::Get().FindTagNode(*RoleName))
+	{
+		bResult = IGameplayTagsEditorModule::Get().DeleteTagFromINI(Node);
+		BroadcastRolesChanged();
+	}
+	return bResult;
+}
+
+#endif
+
+bool UVirtualProductionRolesSubsystem::IsUsingCommandLineRoles() const
+{
+	return bUseCommandLineRoles; 
+}
+
+bool UVirtualProductionRolesSubsystem::HasCommandLineRoles() const
+{
+	return bHasCommandLineRoles;
+}
+
+void UVirtualProductionRolesSubsystem::UseCommandLineRoles(bool bInUseCommandLineRoles)
+ {
+	if (bHasCommandLineRoles)
+	{
+		bUseCommandLineRoles = bInUseCommandLineRoles;
+		BroadcastRolesChanged();
+	}
+}
+
+const TSet<FName>& UVirtualProductionRolesSubsystem::GetRoleSources() const
+{
+	return GetDefault<UVPRolesUserSettings>()->RoleSources;
+}
+
+FName UVirtualProductionRolesSubsystem::GetDefaultRoleSource() const
+{
+	return DefaultVPRolesFileName;
+}
+
+bool UVirtualProductionRolesSubsystem::HasActiveRole(const FGameplayTag& Role) const
+{
+	return GetRolesContainer_Private().HasTag(Role);
+}
+
+void UVirtualProductionRolesSubsystem::ImportCurrentRolesFromDeprecatedSetting()
+{
+	UVPSettings* DeprecatedSettings = GetMutableDefault<UVPSettings>();
+	GetMutableDefault<UVPRolesUserSettings>()->Roles.AppendTags(DeprecatedSettings->Roles);
+	DeprecatedSettings->Roles.Reset();
+	DeprecatedSettings->SaveConfig();
+}
+
+const FGameplayTagContainer& UVirtualProductionRolesSubsystem::GetRolesContainer_Private() const
+{
+	if (bUseCommandLineRoles)
+	{
+		return *CommandLineRoles;
+	}
+		
+	return GetMutableDefault<UVPRolesUserSettings>()->Roles;
+}
+
+const FGameplayTagContainer& UVirtualProductionRolesSubsystem::GetRolesContainerFallback()
+{
+	static const FGameplayTagContainer FallbackGetterContainer;
+	return FallbackGetterContainer;
+}
+
+FGameplayTagContainer* UVirtualProductionRolesSubsystem::GetRolesContainerPtr(EGetRolesPtrSource Source)
+{
+	if (Source == EGetRolesPtrSource::Settings)
+	{
+		return &GetMutableDefault<UVPRolesUserSettings>()->Roles;
+	}
+	else
+	{
+		return CommandLineRoles.Get();
+	}
+}
+
+void UVirtualProductionRolesSubsystem::ImportRoleSourcesFromVPSettings()
+{
+#if WITH_EDITOR
+	const UVPSettings* VPSettings = GetDefault<UVPSettings>();
+	for (const FGameplayTag& Tag : VPSettings->Roles)
+	{
+		FString DevComment;
+		FName TagSource;
+		bool bIsTagExplicit;
+		bool bIsRestrictedTag;
+		bool bAllowNonRestrictedChildren;
+		if (UGameplayTagsManager::Get().GetTagEditorData(Tag.GetTagName(), DevComment, TagSource, bIsTagExplicit, bIsRestrictedTag, bAllowNonRestrictedChildren))
+		{
+			AddRoleSource(TagSource);
+		}
+	}
+#endif
+}
+
+bool UVirtualProductionRolesSubsystem::ImportCurrentRolesFromCommandLine()
+{
+	FString ComandlineRoles;
+
+	if (FParse::Value(FCommandLine::Get(), TEXT("-VPRole="), ComandlineRoles))
+	{
+		bHasCommandLineRoles = true;
+		bUseCommandLineRoles = true;
+		
+		TArray<FString> RoleList;
+		ComandlineRoles.ParseIntoArray(RoleList, TEXT("|"), true);
+
+		for (const FString& Role : RoleList)
+		{
+			FGameplayTag Tag = FGameplayTag::RequestGameplayTag(*Role, false);
+			if (Tag.IsValid())
+			{
+				CommandLineRoles->AddTag(Tag);
+			}
+			else
+			{
+				UE_LOG(LogVPRoles, Error, TEXT("Role %s doesn't exist."), *Role);
+			}
+		}
+		return true;
+	}
+	return false;
+}
+
+void UVirtualProductionRolesSubsystem::BroadcastRolesChanged() const
+{
+#if WITH_EDITOR
+	GetMutableDefault<UVPRolesUserSettings>()->CachedRoles = GetAllRoles();
+	
+	if (GEditor)
+	{
+		// Only save out the active roles if running in editor.
+		GetMutableDefault<UVPRolesUserSettings>()->SaveConfig();
+	}
+#endif
+
+	const TArray<FString> Roles = GetActiveRoles();
+	OnRolesChangedNative.Broadcast(Roles);
+	OnRolesChangedBP.Broadcast(Roles);
+
+#if WITH_EDITOR
+PRAGMA_DISABLE_DEPRECATION_WARNINGS
+	// Trigger a broadcast on the old system.
+	GetMutableDefault<UVPSettings>()->OnRolesChanged.Broadcast();
+PRAGMA_ENABLE_DEPRECATION_WARNINGS
+#endif
+}
+
+void UVirtualProductionRolesSubsystem::AddRoleSource(FName TagSource)
+{	
+	GetMutableDefault<UVPRolesUserSettings>()->RoleSources.Add(TagSource);
+}
+
diff --git a/Engine/Plugins/Experimental/VirtualProductionUtilities/Source/VPRoles/Public/VPRolesSubsystem.h b/Engine/Plugins/Experimental/VirtualProductionUtilities/Source/VPRoles/Public/VPRolesSubsystem.h
new file mode 100644
--- /dev/null
+++ b/Engine/Plugins/Experimental/VirtualProductionUtilities/Source/VPRoles/Public/VPRolesSubsystem.h
@@ -0,0 +1,172 @@
+// Copyright Epic Games, Inc. All Rights Reserved.
+
+#pragma once
+
+#include "CoreMinimal.h"
+#include "Subsystems/EngineSubsystem.h"
+#include "Templates/PimplPtr.h"
+
+#include "VPRolesSubsystem.generated.h"
+
+struct FGameplayTag;
+struct FGameplayTagContainer;
+
+/**
+ * Subsystem for common Virtual Production roles operations.
+ * The machine role(s) in a virtual production context read from the command line.
+ * ie. "-VPRole=[Role.SubRole1|Role.SubRole2]"
+ */
+UCLASS()
+class VPROLES_API UVirtualProductionRolesSubsystem : public UEngineSubsystem
+{
+	GENERATED_BODY()
+
+public:
+	DECLARE_DYNAMIC_MULTICAST_DELEGATE_OneParam(FOnRolesChanged, const TArray<FString>&, EnabledRoles);
+	DECLARE_MULTICAST_DELEGATE_OneParam(FOnRolesChangedNative, const TArray<FString>&);
+	
+public:
+	//~ Begin USubsystem interface
+	virtual void Initialize(FSubsystemCollectionBase& Collection) override;
+	virtual void Deinitialize() override;
+	//~ End USubsystem interface 
+
+	/**
+	 * Get the currently active Virtual Production roles.
+	 * @return The names of the current virtual production roles.
+	 */
+	UFUNCTION(BlueprintCallable, Category = "Virtual Production|Roles")
+	TArray<FString> GetActiveRoles() const;
+
+	/**
+	 * @return Whether a role is currently active.
+	 */
+	UFUNCTION(BlueprintCallable, Category = "Virtual Production|Roles")
+	bool HasActiveRole(const FString& Role) const;
+
+	/**
+	 * @return the currently active Virtual Production roles as a comma separated string representation.
+	 */
+	UFUNCTION(BlueprintCallable, Category = "Virtual Production|Roles")
+	FString GetActiveRolesString() const;
+
+	/**
+	 * Set the current Virtual Production roles.
+	 * @param Roles the list of roles to set as active.
+	 * @note The roles must already have been added either through the UI or through AddRole.
+	 * @note This will clear the previous current roles.
+	 */
+	UFUNCTION(BlueprintCallable, Category = "Virtual Production|Roles")
+	void SetActiveRoles(const TArray<FString>& Roles);
+
+	/**
+	 * Get all available roles that can be set as current.
+	 * @return The available roles.
+	 */
+	UFUNCTION(BlueprintCallable, Category = "Virtual Production|Roles")
+	TArray<FString> GetAllRoles() const;
+	
+#if WITH_EDITOR
+	/**
+	 * Add a new virtual production role.
+	 * @param RoleName the name of the role to add.
+	 * @return Whether the operation succeeded.
+	 * @note This will attempt to modify the underlying VProles.ini file so it must be checked out or made writable or the operation will fail.
+	 */
+	UFUNCTION(BlueprintCallable, Category = "Virtual Production|Roles")
+	bool AddRole(const FString& RoleName);
+
+	/**
+	 * Remove a virtual production role.
+	 * @param RoleName the name of the role to remove.
+	 * @return Whether the operation succeeded.
+	 * @note This will attempt to modify the underlying VProles.ini file so it must be checked out or made writable or the operation will fail.
+	 */
+	UFUNCTION(BlueprintCallable, Category = "Virtual Production|Roles")
+	bool RemoveRole(const FString& RoleName);
+#endif
+
+	/**
+	 * Delegate called when the roles are modified, ie. role is added, current roles are replaced, etc.
+	 */
+	FOnRolesChangedNative& OnRolesChanged()
+	{
+		return OnRolesChangedNative;
+	}
+
+	/**
+	 * @return Whether we are using the roles passed in command line if any.  
+	 */
+	bool IsUsingCommandLineRoles() const;
+
+	/**
+	 * @return Whether roles were passed through the command line.
+	 */
+	bool HasCommandLineRoles() const;
+
+	/**
+	 * Set whether the command line roles should be used if any were specified.
+	 */
+	void UseCommandLineRoles(bool bUseCommandLineRoles);
+
+	/**
+	 * @return The list of sources that provide VP roles.
+	 */
+	const TSet<FName>& GetRoleSources() const;
+	
+	/**
+	 * @return The name of the default VP role ini file.
+	 */
+	FName GetDefaultRoleSource() const;
+
+	/**
+	 * @return Whether a role is currently active.
+	 */
+	bool HasActiveRole(const FGameplayTag& Role) const;
+	
+	/**
+	 * Get the roles container, either from the command line or from the VP Roles settings.
+	 * @note New code should not use this method, it is only provided for backwards compatibility.
+	 */
+	const FGameplayTagContainer& GetRolesContainer_Private() const;
+public:
+	/**
+	 * Delegate called when the roles are modified, ie. role is added, current roles are replaced, etc.
+	 */
+	UPROPERTY(BlueprintAssignable, Category = "Virtual Production|Roles", meta = (ScriptName ="OnRolesChanged"))
+	FOnRolesChanged OnRolesChangedBP;
+	
+private:
+	/** Save roles to the VP Role config file and trigger a broadcast on the OnRolesChanged delegates. */
+	void BroadcastRolesChanged() const;
+	/** Add a role source to the VP Roles settings. Used to keep track of which config files provided VP Roles. */
+	void AddRoleSource(FName TagSource);
+	/** Parses command line roles. */
+	bool ImportCurrentRolesFromCommandLine();
+	/** Imports VP Roles from the VP Settings. */
+	void ImportCurrentRolesFromDeprecatedSetting();
+	/** Fallback handler that returns an empty tag container. */
+	static const FGameplayTagContainer& GetRolesContainerFallback();
+	enum class EGetRolesPtrSource
+	{
+		Settings,
+		CommandLine
+	};
+	/** Get a pointer to the roles gameplay tag container */
+	FGameplayTagContainer* GetRolesContainerPtr(EGetRolesPtrSource Source);
+	/** Gather all gameplay tag sources from the tags that were imported from the VP settings. */
+	void ImportRoleSourcesFromVPSettings();
+
+private:
+	/** Delegate called when the roles are modified, ie. role is added, current roles are replaced, etc. */
+	FOnRolesChangedNative OnRolesChangedNative;
+	/** Holds the role container populated from the roles specified on the command line. */
+	TPimplPtr<FGameplayTagContainer> CommandLineRoles;
+	/** Whether command line roles should be used. */
+	bool bUseCommandLineRoles = false;
+	/** Whether command line roles were specified.  */
+	bool bHasCommandLineRoles = false;
+	
+	friend class FVPRolesModule;
+	friend class FVPRolesEditorModule;
+};
diff --git a/Engine/Plugins/Experimental/VirtualProductionUtilities/Source/VPRoles/VPRoles.Build.cs b/Engine/Plugins/Experimental/VirtualProductionUtilities/Source/VPRoles/VPRoles.Build.cs
new file mode 100644
--- /dev/null
+++ b/Engine/Plugins/Experimental/VirtualProductionUtilities/Source/VPRoles/VPRoles.Build.cs
@@ -0,0 +1,37 @@
+// Copyright Epic Games, Inc. All Rights Reserved.
+
+using UnrealBuildTool;
+
+public class VPRoles : ModuleRules
+{
+	public VPRoles(ReadOnlyTargetRules Target) : base(Target)
+	{
+		PublicDependencyModuleNames.AddRange(
+			new string[]
+			{
+				"Core",
+				"Engine",
+				"CoreUObject",
+				"GameplayTags",
+			}
+		);
+
+		PrivateDependencyModuleNames.AddRange(
+			new string[]
+			{
+				"VPUtilities"
+			}
+		);
+
+		if (Target.bBuildEditor)
+		{
+			PrivateDependencyModuleNames.AddRange(
+				new string[]
+				{
+					"GameplayTagsEditor",
+					"UnrealEd",
+				}
+			);
+		}
+	}
+}
diff --git a/Engine/Plugins/Experimental/VirtualProductionUtilities/Source/VPRolesEditor/Private/VPRolesEditorModule.cpp b/Engine/Plugins/Experimental/VirtualProductionUtilities/Source/VPRolesEditor/Private/VPRolesEditorModule.cpp
new file mode 100644
--- /dev/null
+++ b/Engine/Plugins/Experimental/VirtualProductionUtilities/Source/VPRolesEditor/Private/VPRolesEditorModule.cpp
@@ -0,0 +1,409 @@
+// Copyright Epic Games, Inc. All Rights Reserved.
+
+#include "VPRolesEditorModule.h"
+
+#include "GameplayTagContainer.h"
+#include "GameplayTagsEditorModule.h"
+#include "LevelEditor.h"
+#include "HAL/PlatformFileManager.h"
+#include "Misc/MessageDialog.h"
+#include "SGameplayTagWidget.h"
+#include "SSettingsEditorCheckoutNotice.h"
+#include "ToolMenus.h"
+#include "ToolMenuSection.h"
+#include "VPRolesEditorStyle.h"
+#include "VPRolesSubsystem.h"
+#include "Widgets/Input/SCheckBox.h"
+#include "Widgets/SNullWidget.h"
+#include "Widgets/Input/SEditableTextBox.h"
+#include "Widgets/Layout/SWidgetSwitcher.h"
+
+
+DEFINE_LOG_CATEGORY(LogVPRolesEditor)
+#define LOCTEXT_NAMESPACE "VPRolesEditor"
+
+namespace UE::Private::VPRolesEditor
+{
+	SGameplayTagWidget::ETagFilterResult OnFilterTag(const TSharedPtr<FGameplayTagNode>& InGameplayTagNode)
+	{
+		if (GEngine)
+		{
+			const TSet<FName>& TagSources = GEngine->GetEngineSubsystem<UVirtualProductionRolesSubsystem>()->GetRoleSources();
+			
+			FString DevComment;
+			FName TagSource;
+			bool bIsExplicit = false;
+			bool bIsRestricted = false;
+			bool bAllowNonRestrictedChildren = false;
+
+			static const FString VPRoleComment = TEXT("VPRole");
+			
+			UGameplayTagsManager::Get().GetTagEditorData(InGameplayTagNode->GetCompleteTagName(), DevComment, TagSource, bIsExplicit, bIsRestricted, bAllowNonRestrictedChildren);
+			
+			if ((TagSources.Contains(TagSource) || DevComment == VPRoleComment))
+			{
+				return SGameplayTagWidget::ETagFilterResult::IncludeTag;
+			}
+		}
+		
+		return SGameplayTagWidget::ETagFilterResult::ExcludeTag;
+	};
+	
+}
+
+void FVPRolesEditorModule::StartupModule()
+{
+	FVPRolesEditorStyle::Initialize();
+	FCoreDelegates::OnPostEngineInit.AddRaw(this, &FVPRolesEditorModule::ExtendLevelEditorToolbar);
+}
+
+void FVPRolesEditorModule::ShutdownModule()
+{
+	FCoreDelegates::OnPostEngineInit.RemoveAll(this);
+	FVPRolesEditorStyle::Shutdown();
+}
+
+void FVPRolesEditorModule::ExtendLevelEditorToolbar()
+{
+	if (UVirtualProductionRolesSubsystem* RolesSubsystem = GEngine->GetEngineSubsystem<UVirtualProductionRolesSubsystem>())
+	{
+		UToolMenu* Menu = UToolMenus::Get()->ExtendMenu("LevelEditor.LevelEditorToolBar.User");
+		
+		FToolMenuEntry VPRolesEntry = FToolMenuEntry::InitComboButton(
+			"VPRolesMenu",
+			FUIAction(),
+			FOnGetContent::CreateRaw(this, &FVPRolesEditorModule::GenerateVPRolesLevelEditorToolbarMenu),
+			LOCTEXT("LevelEditorToolbarVPRolesButtonLabel", "VP Roles"),
+			LOCTEXT("LevelEditorToolbarVPRolesButtonTooltip", "Edit VP Roles"),
+			FSlateIcon(FVPRolesEditorStyle::Get().GetStyleSetName(), "VPRolesEditor.TabIcon"));
+		
+		VPRolesEntry.StyleNameOverride = "CalloutToolbar";
+		Menu->FindOrAddSection("VPRoles").AddEntry(VPRolesEntry);
+	}
+}
+
+TSharedRef<SWidget> FVPRolesEditorModule::GenerateVPRolesLevelEditorToolbarMenu()
+{
+	bShowAddRoleTextBox = false;
+	UVirtualProductionRolesSubsystem* RolesSubsystem = GEngine->GetEngineSubsystem<UVirtualProductionRolesSubsystem>();
+	if (!RolesSubsystem)
+	{
+		return SNullWidget::NullWidget;
+	}
+
+	FLevelEditorModule& LevelEditorModule = FModuleManager::LoadModuleChecked<FLevelEditorModule>("LevelEditor");
+	TSharedPtr<FUICommandList> CommandBindings = LevelEditorModule.GetGlobalLevelEditorActions();
+
+	constexpr bool bShouldCloseMenuAfterSelection = false;
+	FMenuBuilder MenuBuilder(bShouldCloseMenuAfterSelection, CommandBindings);
+
+	constexpr bool bNoIndent = false;
+	constexpr bool bSearchable = false;
+	static const FName NoExtensionHook = NAME_None;
+	
+	{
+		MenuBuilder.BeginSection(NoExtensionHook, LOCTEXT("ModifyRolesSection", "Modify Roles"));
+		FMenuEntryParams MenuEntryParams;
+
+		FUIAction AddRoleAction;
+		AddRoleAction.ExecuteAction = FExecuteAction::CreateLambda([this]()
+		{
+			if (TSharedPtr<SWidget> AddRoleWidget = AddRoleEditableBoxWeakPtr.Pin())
+			{
+				bShowAddRoleTextBox = true;
+				FSlateApplication::Get().SetKeyboardFocus(AddRoleWidget);
+			}
+			return;
+		});
+		
+		AddRoleAction.CanExecuteAction = FCanExecuteAction::CreateLambda([this, RolesSubsystem]()
+		{
+			return !bShowAddRoleTextBox && !RolesSubsystem->IsUsingCommandLineRoles();
+		});
+
+		MenuBuilder.AddMenuEntry(LOCTEXT("AddRoleEntry", "Add Role"),
+			LOCTEXT("AddRoleToolTip", "Add a new virtual production role."),
+			FSlateIcon(FVPRolesEditorStyle::Get().GetStyleSetName(), "VPRolesEditor.AddRole"),
+			AddRoleAction);
+		
+		MenuBuilder.EndSection();
+	}
+	
+	{
+		MenuBuilder.BeginSection(NoExtensionHook, LOCTEXT("RolesSection", "Roles"));
+		MenuBuilder.AddWidget(GenerateGameplayTagWidget(), FText::GetEmpty(), bNoIndent, bSearchable);
+		MenuBuilder.EndSection();
+	}
+	
+	return MenuBuilder.MakeWidget();
+}
+
+TSharedRef<SWidget> FVPRolesEditorModule::GenerateGameplayTagWidget()
+{
+	UVirtualProductionRolesSubsystem* RolesSubsystem = GEngine->GetEngineSubsystem<UVirtualProductionRolesSubsystem>();
+	if (!RolesSubsystem)
+	{
+		return SNullWidget::NullWidget;
+	}
+	
+	if (DefaultTagSourceConfigName.IsEmpty())
+	{
+		DefaultTagSourceConfigName = RolesSubsystem->GetDefaultRoleSource().ToString();
+	}
+	
+	static FSlateColorBrush TransparentBrush{ FLinearColor::Transparent };
+	
+	TArray<SGameplayTagWidget::FEditableGameplayTagContainerDatum> EditableContainers;
+	FGameplayTagContainer* ContainerPtr = RolesSubsystem->GetRolesContainerPtr(UVirtualProductionRolesSubsystem::EGetRolesPtrSource::Settings);
+	EditableContainers.Emplace(nullptr, ContainerPtr);
+
+	TArray<SGameplayTagWidget::FEditableGameplayTagContainerDatum> CommandLineEditableContainers;
+    FGameplayTagContainer* CommandLineContainerPtr = RolesSubsystem->GetRolesContainerPtr(UVirtualProductionRolesSubsystem::EGetRolesPtrSource::CommandLine);
+    CommandLineEditableContainers.Emplace(nullptr, CommandLineContainerPtr);
+
+	return SNew(SBorder)
+		.BorderImage(FAppStyle::GetBrush("NoBorder"))
+		.Padding(FMargin(8.f, 0.f))
+		[
+			SNew(SVerticalBox)
+			+ SVerticalBox::Slot()
+			.Padding(2.f)
+			.AutoHeight()
+			[
+				SNew(SBorder)
+				.Visibility_Lambda([RolesSubsystem]() { return RolesSubsystem->HasCommandLineRoles() ? EVisibility::Visible : EVisibility::Collapsed; })
+				.Padding(FMargin(6.f, 4.f))
+				.BorderImage(FAppStyle::GetBrush("ToolPanel.GroupBorder"))
+				[
+					SNew(SHorizontalBox)
+					+ SHorizontalBox::Slot()
+					.AutoWidth()
+					.VAlign(VAlign_Center)
+					[
+						SNew(SCheckBox)
+						.IsChecked_Lambda([RolesSubsystem](){ return RolesSubsystem->IsUsingCommandLineRoles() ? ECheckBoxState::Checked : ECheckBoxState::Unchecked; })
+						.OnCheckStateChanged_Lambda([RolesSubsystem, this](ECheckBoxState CheckBoxState)
+						{
+							RolesSubsystem->UseCommandLineRoles(CheckBoxState == ECheckBoxState::Checked ? true : false);
+						})
+					]
+					+ SHorizontalBox::Slot()
+					.FillWidth(1.f)
+					.VAlign(VAlign_Center)
+					[
+						SNew(STextBlock)	
+						.Text(LOCTEXT("UseCommandLineRolesLabel", "Use command line roles"))
+					]
+				]
+			]
+			+ SVerticalBox::Slot()
+			.AutoHeight()
+			[
+				SNew(SHorizontalBox)
+				.Visibility_Lambda(	[RolesSubsystem, this](){ return !bShowAddRoleTextBox || RolesSubsystem->IsUsingCommandLineRoles() ? EVisibility::Collapsed : EVisibility::Visible; })
+				+SHorizontalBox::Slot()
+				.AutoWidth()
+				[
+					// To match the padding of the gameplay tag list.
+					SNew(SImage)
+					.Image(FAppStyle::GetBrush("TreeArrow_Collapsed"))
+					.Visibility(EVisibility::Hidden)
+				]
+				+SHorizontalBox::Slot()
+	            .AutoWidth()
+	            [
+	                SNew(SCheckBox)
+	                .IsEnabled(false)
+	            ]
+				+SHorizontalBox::Slot()
+				.Padding(0.f)
+				[
+					SAssignNew(AddRoleEditableBoxWeakPtr, SEditableTextBox)
+						.IsReadOnly_Lambda([RolesSubsystem](){ return RolesSubsystem->IsUsingCommandLineRoles(); })
+						.ClearKeyboardFocusOnCommit(false)
+						.HintText(LOCTEXT("GameplayTagWidget_AddRole", "Add Role"))
+						.OnTextChanged_Raw(this, &FVPRolesEditorModule::OnAddTagTextChanged)
+						.OnTextCommitted_Raw( this, &FVPRolesEditorModule::OnSubmitNewTag)
+				]
+			]
+			+ SVerticalBox::Slot()
+			.Padding(0.f)
+			.AutoHeight()
+			[
+				SNew(SBox)
+				.MinDesiredWidth(200.f)
+				[
+					SNew(SWidgetSwitcher)
+					.WidgetIndex_Lambda([RolesSubsystem]()
+					{
+						return RolesSubsystem->IsUsingCommandLineRoles() ? 1 : 0;
+					})
+					+SWidgetSwitcher::Slot()
+					[
+						SNew(SGameplayTagWidget, EditableContainers)
+							.Padding(0.f)
+							.OnFilterTag_Static(&UE::Private::VPRolesEditor::OnFilterTag)
+							.ForceHideAddNewTag(true)
+							.ForceHideAddNewTagSource(true)
+							.ForceHideTagTreeControls(true)
+							.GameplayTagUIMode(EGameplayTagUIMode::HybridMode)
+							.MaxHeight(0) // Height of 0 means the height is not limited.
+							.BackgroundBrush(&TransparentBrush)
+							.TagTreeViewBackgroundBrush(&TransparentBrush)
+							.OnTagChanged_Lambda([RolesSubsystem]{ RolesSubsystem->BroadcastRolesChanged(); })
+					]
+					+SWidgetSwitcher::Slot()
+					[
+						SNew(SGameplayTagWidget, CommandLineEditableContainers)
+							.Padding(0.f)
+							.OnFilterTag_Static(&UE::Private::VPRolesEditor::OnFilterTag)
+							.ForceHideAddNewTag(true)
+							.ForceHideAddNewTagSource(true)
+							.ForceHideTagTreeControls(true)
+							.GameplayTagUIMode(EGameplayTagUIMode::SelectionMode)
+							.MaxHeight(0) // Height of 0 means the height is not limited.
+							.BackgroundBrush(&TransparentBrush)
+							.TagTreeViewBackgroundBrush(&TransparentBrush)
+							.ReadOnly(true)
+					]					
+				]
+			]
+			+SVerticalBox::Slot()
+			.AutoHeight()
+			[
+				SNew(SSettingsEditorCheckoutNotice)
+				.Visibility_Raw(this, &FVPRolesEditorModule::DetermineFileWatcherWidgetVisibility)
+				.ConfigFilePath(RolesSubsystem->GetDefaultRoleSource().ToString())
+			]
+		];
+}
+
+void FVPRolesEditorModule::OnAddTagTextChanged(const FText& Text)
+{
+	if (const TSharedPtr<SEditableTextBox> AddTagEditableBox = AddRoleEditableBoxWeakPtr.Pin())
+	{
+		const FString TagName = Text.ToString();
+		FText ErrorMsg;
+
+		if (Text.IsEmpty())
+		{
+			AddTagEditableBox->SetError(FText::GetEmpty());
+			return;
+		}
+
+		if (!UGameplayTagsManager::Get().IsValidGameplayTagString(TagName, &ErrorMsg))
+		{
+			FText MessageTitle(LOCTEXT("InvalidTag", "Invalid Tag"));
+			AddTagEditableBox->SetError(ErrorMsg);
+		}
+		else
+		{
+			AddTagEditableBox->SetError(FText::GetEmpty());
+		}
+	}
+}
+
+EVisibility FVPRolesEditorModule::DetermineFileWatcherWidgetVisibility() const
+{
+	return CanModifyTagSource() ? EVisibility::Collapsed : EVisibility::Visible;
+}
+
+bool FVPRolesEditorModule::CanModifyTagSource() const
+{
+	if (!DefaultTagSourceConfigName.IsEmpty())
+	{
+		return !FPlatformFileManager::Get().GetPlatformFile().IsReadOnly(*DefaultTagSourceConfigName);
+	}
+	else
+	{
+		ensureMsgf(false, TEXT("File %s should already exist."), *DefaultTagSourceConfigName);
+		return false;
+	}
+}
+
+void FVPRolesEditorModule::OnSubmitNewTag(const FText& CommittedText, ETextCommit::Type CommitType)
+{
+	if (!GEngine)
+	{
+		return;
+	}
+	
+	const TSharedPtr<SEditableTextBox> AddTagEditableBox = AddRoleEditableBoxWeakPtr.Pin();
+	if (!AddTagEditableBox)
+	{
+		return;
+	}
+	
+	if (const UVirtualProductionRolesSubsystem* RolesSubsystem = GEngine->GetEngineSubsystem<UVirtualProductionRolesSubsystem>())
+	{
+		if (CommitType == ETextCommit::Type::OnEnter)
+		{
+			UGameplayTagsManager& Manager = UGameplayTagsManager::Get();
+
+			// Only support adding tags via ini file
+			if (Manager.ShouldImportTagsFromINI() == false)
+			{
+				return;
+			}
+
+			const FText TagNameAsText = CommittedText;
+			const FString TagName = TagNameAsText.ToString();
+
+			if (TagName.IsEmpty())
+			{
+				return;
+			}
+
+			if (DefaultTagSourceConfigName.IsEmpty())
+			{
+				return;
+			}
+
+			// Check to see if this is a valid tag
+			// First check the base rules for all tags then look for any additional rules in the delegate
+			FText ErrorMsg;
+			if (!Manager.IsValidGameplayTagString(TagName, &ErrorMsg))
+			{
+				AddTagEditableBox->SetError(ErrorMsg);
+				return;
+			}
+
+			bool bCanAddTag = false;
+			if (!CanModifyTagSource())
+			{
+				if (FMessageDialog::Open(EAppMsgType::YesNo, LOCTEXT("SaveAsDefaultsIsReadOnlyMessage", "The ini file for these settings is not currently writable. Would you like to make it writable?")) == EAppReturnType::Yes)
+				{
+					bCanAddTag = SettingsHelpers::MakeWritable(DefaultTagSourceConfigName, true);
+				}
+			}
+			else
+			{
+				bCanAddTag = true;
+				AddTagEditableBox->SetText(FText::GetEmpty());
+			}
+
+			if (bCanAddTag)
+			{
+				if (IGameplayTagsEditorModule::Get().AddNewGameplayTagToINI(TagName, FString(), *DefaultTagSourceConfigName))
+				{
+					if (TSharedPtr<FGameplayTagNode> Node = Manager.FindTagNode(*TagName))
+					{
+						RolesSubsystem->BroadcastRolesChanged();
+					}
+				}
+				
+				bShowAddRoleTextBox = false;
+			}
+		}
+
+		if (CommitType == ETextCommit::Type::OnUserMovedFocus)
+		{
+			bShowAddRoleTextBox = false;
+			AddTagEditableBox->SetText(FText::GetEmpty());
+		}
+	}
+}
+
+IMPLEMENT_MODULE(FVPRolesEditorModule, VPRolesEditor)
+
+#undef LOCTEXT_NAMESPACE /*VPRolesEditor*/
diff --git a/Engine/Plugins/Experimental/VirtualProductionUtilities/Source/VPRolesEditor/Private/VPRolesEditorStyle.cpp b/Engine/Plugins/Experimental/VirtualProductionUtilities/Source/VPRolesEditor/Private/VPRolesEditorStyle.cpp
new file mode 100644
--- /dev/null
+++ b/Engine/Plugins/Experimental/VirtualProductionUtilities/Source/VPRolesEditor/Private/VPRolesEditorStyle.cpp
@@ -0,0 +1,64 @@
+// Copyright Epic Games, Inc. All Rights Reserved.
+
+#include "VPRolesEditorStyle.h"
+#include "Interfaces/IPluginManager.h"
+#include "Framework/Application/SlateApplication.h"
+#include "Slate/SlateGameResources.h"
+#include "Styling/SlateStyleRegistry.h"
+#include "Styling/StyleColors.h"
+
+TSharedPtr<FSlateStyleSet> FVPRolesEditorStyle::StyleInstance;
+
+void FVPRolesEditorStyle::Initialize()
+{
+	if (!StyleInstance.IsValid())
+	{
+		StyleInstance = Create();
+		FSlateStyleRegistry::RegisterSlateStyle(*StyleInstance);
+	}
+}
+
+void FVPRolesEditorStyle::Shutdown()
+{
+	FSlateStyleRegistry::UnRegisterSlateStyle(*StyleInstance);
+	ensure(StyleInstance.IsUnique());
+	StyleInstance.Reset();
+}
+
+FName FVPRolesEditorStyle::GetStyleSetName()
+{
+	static FName StyleSetName(TEXT("VPRolesEditorStyle"));
+	return StyleSetName;
+}
+
+#define IMAGE_BRUSH( RelativePath, ... ) FSlateImageBrush( Style->RootToContentDir( RelativePath, TEXT(".png") ), __VA_ARGS__ )
+
+const FVector2D Icon16x16(16.0f, 16.0f);
+const FVector2D Icon40x40(40.0f, 40.0f);
+
+TSharedRef<FSlateStyleSet> FVPRolesEditorStyle::Create()
+{
+	TSharedRef<FSlateStyleSet> Style = MakeShared<FSlateStyleSet>(FVPRolesEditorStyle::GetStyleSetName());
+	Style->SetContentRoot(IPluginManager::Get().FindPlugin("VirtualProductionUtilities")->GetBaseDir() / TEXT("Resources"));
+
+	FSlateBrush AddRoleBrush = *FAppStyle::Get().GetBrush("Icons.Plus");
+	AddRoleBrush.TintColor = FStyleColors::AccentGreen;
+
+	Style->Set("VPRolesEditor.TabIcon", new IMAGE_BRUSH(TEXT("VPRolesButtonIcon_40x"), Icon16x16));
+	Style->Set("VPRolesEditor.OpenMenu", new IMAGE_BRUSH(TEXT("VPRolesButtonIcon_40x"), Icon40x40));
+	Style->Set("VPRolesEditor.AddRole", new FSlateBrush(AddRoleBrush));
+
+	return Style;
+}
+
+const ISlateStyle& FVPRolesEditorStyle::Get()
+{
+	static const FVPRolesEditorStyle Inst;
+	return *StyleInstance;
+}
+
+#undef IMAGE_BRUSH
+#undef BOX_BRUSH
+#undef BORDER_BRUSH
+#undef TTF_FONT
+#undef OTF_FONT
diff --git a/Engine/Plugins/Experimental/VirtualProductionUtilities/Source/VPRolesEditor/Private/VPRolesEditorStyle.h b/Engine/Plugins/Experimental/VirtualProductionUtilities/Source/VPRolesEditor/Private/VPRolesEditorStyle.h
new file mode 100644
--- /dev/null
+++ b/Engine/Plugins/Experimental/VirtualProductionUtilities/Source/VPRolesEditor/Private/VPRolesEditorStyle.h
@@ -0,0 +1,27 @@
+// Copyright Epic Games, Inc. All Rights Reserved.
+
+#pragma once
+
+#include "CoreMinimal.h"
+#include "Styling/SlateStyle.h"
+
+class FVPRolesEditorStyle
+{
+public:
+	/** Initialize the style. */
+	static void Initialize();
+	/** Unregister the style. */
+	static void Shutdown();
+	/** The Slate style set for VP Roles */
+	static const ISlateStyle& Get();
+	/** Get the style set name. */
+	static FName GetStyleSetName();
+
+private:
+	/** Initializes the singleton and sets the buttons used by the vp roles widget.  */
+	static TSharedRef<class FSlateStyleSet> Create();
+
+private:
+	/** Singleton instance of the style. */
+	static TSharedPtr<class FSlateStyleSet> StyleInstance;
+};
\ No newline at end of file
diff --git a/Engine/Plugins/Experimental/VirtualProductionUtilities/Source/VPRolesEditor/Public/VPRolesEditorModule.h b/Engine/Plugins/Experimental/VirtualProductionUtilities/Source/VPRolesEditor/Public/VPRolesEditorModule.h
new file mode 100644
--- /dev/null
+++ b/Engine/Plugins/Experimental/VirtualProductionUtilities/Source/VPRolesEditor/Public/VPRolesEditorModule.h
@@ -0,0 +1,53 @@
+// Copyright Epic Games, Inc. All Rights Reserved.
+
+#pragma once
+
+#include "CoreMinimal.h"
+#include "Modules/ModuleInterface.h"
+#include "Types/SlateEnums.h"
+
+class SWidget;
+struct EVisibility;
+
+DECLARE_LOG_CATEGORY_EXTERN(LogVPRolesEditor, Log, Log);
+
+class VPROLESEDITOR_API FVPRolesEditorModule : public IModuleInterface
+{
+public:
+	//~ Begin IModuleInterface
+	virtual void StartupModule() override;
+	virtual void ShutdownModule() override;
+	//~End IModuleInterface
+
+private:
+	/** Core delegate handler for extending the level editor toolbar. */
+	void ExtendLevelEditorToolbar();
+	
+	/** Generates the toolbar button widget. */
+	TSharedRef<SWidget> GenerateVPRolesLevelEditorToolbarMenu();
+
+	/** Generate the gameplay tag widget. */
+	TSharedRef<SWidget> GenerateGameplayTagWidget();
+	
+	/** Callback for when a new tag is added */
+	void OnSubmitNewTag(const FText& CommittedText, ETextCommit::Type CommitType);
+
+	/** Callback for when the add tag textbox's text changed. */
+	void OnAddTagTextChanged(const FText& Text);
+	
+	/** Whether the file watcher widget should be shown for the default tag source config file. */
+	EVisibility DetermineFileWatcherWidgetVisibility() const;
+	
+	/** Whether the default VP role tag source can be modified. */
+	bool CanModifyTagSource() const;
+
+private:
+	/** The file name for the default vp role tag source. */
+	FString DefaultTagSourceConfigName;
+	/** Holds the widget generated by the GameplayTags editor module. */
+	TWeakPtr<class SGameplayTagWidget> GameplayTagsWidget;
+	/** Holds the editable text box widget, used to update the text box's error message. */
+	TWeakPtr<class SEditableTextBox> AddRoleEditableBoxWeakPtr;
+	/** Whether to show the add role textbox. */
+	bool bShowAddRoleTextBox = false;
+};
\ No newline at end of file
diff --git a/Engine/Plugins/Experimental/VirtualProductionUtilities/Source/VPRolesEditor/VPRolesEditor.Build.cs b/Engine/Plugins/Experimental/VirtualProductionUtilities/Source/VPRolesEditor/VPRolesEditor.Build.cs
new file mode 100644
--- /dev/null
+++ b/Engine/Plugins/Experimental/VirtualProductionUtilities/Source/VPRolesEditor/VPRolesEditor.Build.cs
@@ -0,0 +1,34 @@
+// Copyright Epic Games, Inc. All Rights Reserved.
+
+using UnrealBuildTool;
+
+public class VPRolesEditor : ModuleRules
+{
+	public VPRolesEditor(ReadOnlyTargetRules Target) : base(Target)
+	{
+		PublicDependencyModuleNames.AddRange(
+			new string[]
+			{
+				"Core",
+			}
+		);
+
+		PrivateDependencyModuleNames.AddRange(
+			new string[]
+			{
+				"CoreUObject",
+				"Engine",
+				"GameplayTags",
+				"GameplayTagsEditor",
+				"LevelEditor",
+				"Projects",
+				"Slate",
+				"SlateCore",
+				"ToolMenus",
+				"VPRoles",
+				"VPUtilities",
+				"SharedSettingsWidgets",
+			}
+		);
+	}
+}
diff --git a/Engine/Plugins/Experimental/VirtualProductionUtilities/Source/VPUtilities/Private/VPBlueprintLibrary.cpp b/Engine/Plugins/Experimental/VirtualProductionUtilities/Source/VPUtilities/Private/VPBlueprintLibrary.cpp
--- a/Engine/Plugins/Experimental/VirtualProductionUtilities/Source/VPUtilities/Private/VPBlueprintLibrary.cpp
+++ b/Engine/Plugins/Experimental/VirtualProductionUtilities/Source/VPUtilities/Private/VPBlueprintLibrary.cpp
@@ -109,7 +109,11 @@ bool UVPBlueprintLibrary::JumpToBookmarkInLevelEditor(const UVPBookmark* Bookmar
 
 FGameplayTagContainer UVPBlueprintLibrary::GetVirtualProductionRole()
 {
+PRAGMA_DISABLE_DEPRECATION_WARNINGS
+	// We can't depend upon the VP Roles subsystem because it will introduce a circular dependency so we have to
+	// depend upon the deprecated method.
 	return GetDefault<UVPSettings>()->GetRoles();
+PRAGMA_ENABLE_DEPRECATION_WARNINGS
 }
 
 
diff --git a/Engine/Plugins/Experimental/VirtualProductionUtilities/Source/VPUtilities/Private/VPSettings.cpp b/Engine/Plugins/Experimental/VirtualProductionUtilities/Source/VPUtilities/Private/VPSettings.cpp
--- a/Engine/Plugins/Experimental/VirtualProductionUtilities/Source/VPUtilities/Private/VPSettings.cpp
+++ b/Engine/Plugins/Experimental/VirtualProductionUtilities/Source/VPUtilities/Private/VPSettings.cpp
@@ -2,48 +2,18 @@
 
 #include "VPSettings.h"
 
-#include "Misc/CommandLine.h"
-#include "VPUtilitiesModule.h"
-
-UVPSettings::UVPSettings()
+const FGameplayTagContainer& UVPSettings::GetRoles() const
 {
-	FString ComandlineRoles;
-	bIsCommandLineRolesValid = FParse::Value(FCommandLine::Get(), TEXT("-VPRole="), ComandlineRoles);
-
-	if (bIsCommandLineRolesValid)
+	static const FGameplayTagContainer EmptyContainer;
+	if (ensure(RolesGetter.IsBound()))
 	{
-		TArray<FString> RoleList;
-		ComandlineRoles.ParseIntoArray(RoleList, TEXT("|"), true);
-
-		for (const FString& Role : RoleList)
-		{
-			FGameplayTag Tag = FGameplayTag::RequestGameplayTag(*Role, false);
-			if (Tag.IsValid())
-			{
-				CommandLineRoles.AddTag(Tag);
-			}
-			else
-			{
-				UE_LOG(LogVPUtilities, Error, TEXT("Role %s doesn't exist."), *Role);
-			}
-		}
+		return RolesGetter.Execute();
 	}
+	
+	return EmptyContainer;
 }
 
 UVPSettings* UVPSettings::GetVPSettings()
 {
 	return GetMutableDefault<UVPSettings>();
 }
-
-#if WITH_EDITOR
-
-void UVPSettings::PostEditChangeChainProperty(struct FPropertyChangedChainEvent& PropertyChangedEvent)
-{
-	if (PropertyChangedEvent.Property->GetFName() == GET_MEMBER_NAME_CHECKED(ThisClass, Roles))
-	{
-		OnRolesChanged.Broadcast();
-	}
-
-	Super::PostEditChangeChainProperty(PropertyChangedEvent);
-}
-#endif
diff --git a/Engine/Plugins/Experimental/VirtualProductionUtilities/Source/VPUtilities/Public/VPSettings.h b/Engine/Plugins/Experimental/VirtualProductionUtilities/Source/VPUtilities/Public/VPSettings.h
--- a/Engine/Plugins/Experimental/VirtualProductionUtilities/Source/VPUtilities/Public/VPSettings.h
+++ b/Engine/Plugins/Experimental/VirtualProductionUtilities/Source/VPUtilities/Public/VPSettings.h
@@ -5,7 +5,6 @@
 #include "CoreMinimal.h"
 #include "GameplayTagContainer.h"
 #include "UObject/Object.h"
-#include "Engine/DeveloperSettings.h"
 #include "VPSettings.generated.h"
 
 /**
@@ -17,26 +16,8 @@ class VPUTILITIES_API UVPSettings : public UObject
 	GENERATED_BODY()
 
 public:
-	UVPSettings();
-
-protected:
-	/**
-	 * The machine role(s) in a virtual production context.
-	 * @note The role may be override via the command line, "-VPRole=[Role.SubRole1|Role.SubRole2]"
-	 */
-	UPROPERTY(config, EditAnywhere, Category="Virtual Production")
-	FGameplayTagContainer Roles;
-
-	/**
-	 * The machine role(s) in a virtual production context read from the command line.
-	 * ie. "-VPRole=[Role.SubRole1|Role.SubRole2]"
-	 */
-	bool bIsCommandLineRolesValid;
-	UPROPERTY(transient)
-	FGameplayTagContainer CommandLineRoles;
-
-public:
-	const FGameplayTagContainer& GetRoles() const { return bIsCommandLineRolesValid ? CommandLineRoles : Roles; }
+	UE_DEPRECATED(5.1, "Use UVirtualProductionRolesSubsystem::GetCurrentVirtualProductionRoles instead.")
+	const FGameplayTagContainer& GetRoles() const;
 
 	/** Default Kit of Focal Lengths for Virtual Camera */
 	UPROPERTY(EditAnywhere, BlueprintReadWrite, config, Category = "VirtualCamera|Presets")
@@ -62,6 +43,7 @@ public:
 	UPROPERTY(config, EditAnywhere, Category="Virtual Production")
 	bool bShowRoleInEditor;
 
+	UE_DEPRECATED(5.1, "Use UVirtualProductionRolesSubsystem::OnRolesChanged instead.")
 	/** Notify when the virtual production roles have changed. */
 	FSimpleMulticastDelegate OnRolesChanged;
 
@@ -70,16 +52,23 @@ public:
 
 	UPROPERTY(config, EditAnywhere, Category = "Virtual Production", DisplayName = "Project Name")
 	FString ShowName;
-
 #endif
+	
+protected:
+	/**
+	 * The machine role(s) in a virtual production context.
+	 * @note The role may be override via the command line, "-VPRole=[Role.SubRole1|Role.SubRole2]"
+	 */
+	UPROPERTY(config)
+	FGameplayTagContainer Roles;
 
-public:
-#if WITH_EDITOR
-	//~ UObject interface
-	virtual void PostEditChangeChainProperty(struct FPropertyChangedChainEvent& PropertyChangedEvent) override;
-	//~ End UObject interface
-
-
-#endif //WITH_EDITOR
+private:
+	/**
+	 * Roles getter set by the VP Roles module to maintain backwards compatibility.
+	 * May return the command line roles or roles enabled through the VP Roles subsystem.
+	 */
+	DECLARE_DELEGATE_RetVal(const FGameplayTagContainer&, FOnGetRoles);
+	FOnGetRoles RolesGetter;
 
+	friend class UVirtualProductionRolesSubsystem;
 };
\ No newline at end of file
diff --git a/Engine/Plugins/Experimental/VirtualProductionUtilities/Source/VPUtilitiesEditor/Private/VPUtilitiesEditorModule.cpp b/Engine/Plugins/Experimental/VirtualProductionUtilities/Source/VPUtilitiesEditor/Private/VPUtilitiesEditorModule.cpp
--- a/Engine/Plugins/Experimental/VirtualProductionUtilities/Source/VPUtilitiesEditor/Private/VPUtilitiesEditorModule.cpp
+++ b/Engine/Plugins/Experimental/VirtualProductionUtilities/Source/VPUtilitiesEditor/Private/VPUtilitiesEditorModule.cpp
@@ -5,6 +5,7 @@
 #include "Framework/Docking/WorkspaceItem.h"
 #include "GameplayTagContainer.h"
 #include "Editor.h"
+#include "Engine/Engine.h"
 #include "EditorUtilitySubsystem.h"
 #include "HAL/IConsoleManager.h"
 #include "LevelEditor.h"
@@ -18,6 +19,7 @@
 #include "SGenlockProviderTab.h"
 #include "STimecodeProviderTab.h"
 #include "Textures/SlateIcon.h"
+#include "VPRolesSubsystem.h"
 #include "VPSettings.h"
 #include "VPUtilitiesEditorSettings.h"
 #include "VPUtilitiesEditorStyle.h"
@@ -125,7 +127,7 @@ void FVPUtilitiesEditorModule::RegisterSettings()
 	{
 		FLevelEditorModule::FTitleBarItem Item;
 		Item.Label = LOCTEXT("VPRolesLabel", "VP Roles: ");
-		Item.Value = MakeAttributeLambda([]() { return FText::FromString(GetMutableDefault<UVPSettings>()->GetRoles().ToStringSimple()); });
+		Item.Value = MakeAttributeLambda([]() { return FText::FromString(GEngine->GetEngineSubsystem<UVirtualProductionRolesSubsystem>()->GetActiveRolesString()); });
 		Item.Visibility = MakeAttributeLambda([]() { return GetMutableDefault<UVPSettings>()->bShowRoleInEditor ? EVisibility::SelfHitTestInvisible : EVisibility::Collapsed; });
 		LevelEditorModule->AddTitleBarItem(VPRoleNotificationBarIdentifier, Item);
 	}
diff --git a/Engine/Plugins/Experimental/VirtualProductionUtilities/Source/VPUtilitiesEditor/VPUtilitiesEditor.Build.cs b/Engine/Plugins/Experimental/VirtualProductionUtilities/Source/VPUtilitiesEditor/VPUtilitiesEditor.Build.cs
--- a/Engine/Plugins/Experimental/VirtualProductionUtilities/Source/VPUtilitiesEditor/VPUtilitiesEditor.Build.cs
+++ b/Engine/Plugins/Experimental/VirtualProductionUtilities/Source/VPUtilitiesEditor/VPUtilitiesEditor.Build.cs
@@ -40,6 +40,7 @@ public class VPUtilitiesEditor : ModuleRules
 				"UMGEditor",
 				"UnrealEd",
 				"VPBookmark",
+				"VPRoles",
 				"WorkspaceMenuStructure",
             }
 		);
diff --git a/Engine/Plugins/Experimental/VirtualProductionUtilities/VirtualProductionUtilities.uplugin b/Engine/Plugins/Experimental/VirtualProductionUtilities/VirtualProductionUtilities.uplugin
--- a/Engine/Plugins/Experimental/VirtualProductionUtilities/VirtualProductionUtilities.uplugin
+++ b/Engine/Plugins/Experimental/VirtualProductionUtilities/VirtualProductionUtilities.uplugin
@@ -38,6 +38,16 @@
 			"Name": "VPMaterialsEditor",
 			"Type": "Editor",
 			"LoadingPhase": "PostDefault"
+		},
+		{
+			"Name": "VPRoles",
+			"Type": "Runtime",
+			"LoadingPhase": "Default"
+		},
+		{
+			"Name": "VPRolesEditor",
+			"Type": "Editor",
+			"LoadingPhase": "Default"
 		}
 	],
 	"Plugins": [
@@ -64,6 +74,11 @@
 		{
 			"Name": "Composure",
 			"Enabled": true
+		},
+		{
+			"Name": "GameplayTagsEditor",
+			"Enabled": true
 		}
+
 	]
 }
diff --git a/Engine/Plugins/VirtualProduction/StageMonitoring/Source/StageDataProvider/Private/StageDataProvider.cpp b/Engine/Plugins/VirtualProduction/StageMonitoring/Source/StageDataProvider/Private/StageDataProvider.cpp
--- a/Engine/Plugins/VirtualProduction/StageMonitoring/Source/StageDataProvider/Private/StageDataProvider.cpp
+++ b/Engine/Plugins/VirtualProduction/StageMonitoring/Source/StageDataProvider/Private/StageDataProvider.cpp
@@ -3,6 +3,7 @@
 #include "StageDataProvider.h"
 
 #include "Containers/Ticker.h"
+#include "Engine/Engine.h"
 #include "MessageEndpoint.h"
 #include "MessageEndpointBuilder.h"
 #include "Misc/App.h"
@@ -11,6 +12,7 @@
 #include "StageMonitoringSettings.h"
 #include "StageMonitorUtils.h"
 #include "VPSettings.h"
+#include "VPRoles/Public/VPRolesSubsystem.h"
 
 
 namespace StageDataProviderUtils
@@ -187,10 +189,13 @@ bool FStageDataProvider::IsMessageTypeExcluded(UScriptStruct* MessageType) const
 	const FStageDataProviderSettings& ProviderSettings = GetDefault<UStageMonitoringSettings>()->ProviderSettings;
 	if (const FGameplayTagContainer* SupportedRoles = ProviderSettings.MessageTypeRoleExclusion.Find(FStageMessageTypeWrapper(MessageType->GetFName())))
 	{
-		const FGameplayTagContainer& CurrentRoles = GetDefault<UVPSettings>()->GetRoles();
-		if (!SupportedRoles->HasAny(CurrentRoles))
+		if (GEngine)
 		{
-			return true;
+			const FGameplayTagContainer& CurrentRoles = GEngine->GetEngineSubsystem<UVirtualProductionRolesSubsystem>()->GetRolesContainer_Private();
+			if (!SupportedRoles->HasAny(CurrentRoles))
+			{
+				return true;
+			}
 		}
 	}
 
diff --git a/Engine/Plugins/VirtualProduction/StageMonitoring/Source/StageDataProvider/Private/StageDataProviderModule.cpp b/Engine/Plugins/VirtualProduction/StageMonitoring/Source/StageDataProvider/Private/StageDataProviderModule.cpp
--- a/Engine/Plugins/VirtualProduction/StageMonitoring/Source/StageDataProvider/Private/StageDataProviderModule.cpp
+++ b/Engine/Plugins/VirtualProduction/StageMonitoring/Source/StageDataProvider/Private/StageDataProviderModule.cpp
@@ -2,6 +2,7 @@
 
 #include "StageDataProviderModule.h"
 
+#include "Engine/Engine.h"
 #include "Features/IModularFeatures.h"
 #include "FramePerformanceProvider.h"
 #include "GenlockWatchdog.h"
@@ -12,6 +13,7 @@
 #include "TakeRecorderStateProvider.h"
 #include "TimecodeProviderWatchdog.h"
 #include "VPSettings.h"
+#include "VPRoles/Public/VPRolesSubsystem.h"
 
 
 const FName IStageDataProviderModule::ModuleName = TEXT("StageDataProvider");
@@ -58,17 +60,18 @@ void FStageDataProviderModule::ShutdownModule()
 
 void FStageDataProviderModule::OnPostEngineInit()
 {
+	const UVirtualProductionRolesSubsystem* RolesSubsystem = GEngine->GetEngineSubsystem<UVirtualProductionRolesSubsystem>();
 	//If role enforcement is enabled, make sure we have one required to enable data provider
 	const UStageMonitoringSettings* Settings = GetDefault<UStageMonitoringSettings>();
-	if (!Settings->ProviderSettings.bUseRoleFiltering || GetDefault<UVPSettings>()->GetRoles().HasAny(Settings->ProviderSettings.SupportedRoles))
+	if (!Settings->ProviderSettings.bUseRoleFiltering || RolesSubsystem->GetRolesContainer_Private().HasAny(Settings->ProviderSettings.SupportedRoles))
 	{
 		StartDataProvider();
 	}
 	else
 	{
 		UE_LOG(LogStageDataProvider, Log, TEXT("Can't start StageDataProvider. Role filtering is enabled and our roles (%s) are not accepted (%s)")
-		, *GetDefault<UVPSettings>()->GetRoles().ToStringSimple()
-		, *Settings->ProviderSettings.SupportedRoles.ToStringSimple())
+		, *RolesSubsystem->GetActiveRolesString()
+		, *Settings->ProviderSettings.SupportedRoles.ToStringSimple());
 	}
 }
 
diff --git a/Engine/Plugins/VirtualProduction/StageMonitoring/Source/StageDataProvider/StageDataProvider.build.cs b/Engine/Plugins/VirtualProduction/StageMonitoring/Source/StageDataProvider/StageDataProvider.build.cs
--- a/Engine/Plugins/VirtualProduction/StageMonitoring/Source/StageDataProvider/StageDataProvider.build.cs
+++ b/Engine/Plugins/VirtualProduction/StageMonitoring/Source/StageDataProvider/StageDataProvider.build.cs
@@ -26,7 +26,8 @@ public class StageDataProvider : ModuleRules
 				"StageDataCore",
 				"StageMonitorCommon",
 				"TimeManagement",
-				"VPUtilities"
+				"VPRoles",
+				"VPUtilities",
 			}
 		);
 
diff --git a/Engine/Plugins/VirtualProduction/StageMonitoring/Source/StageMonitor/Private/StageMonitorModule.cpp b/Engine/Plugins/VirtualProduction/StageMonitoring/Source/StageMonitor/Private/StageMonitorModule.cpp
--- a/Engine/Plugins/VirtualProduction/StageMonitoring/Source/StageMonitor/Private/StageMonitorModule.cpp
+++ b/Engine/Plugins/VirtualProduction/StageMonitoring/Source/StageMonitor/Private/StageMonitorModule.cpp
@@ -2,12 +2,14 @@
 
 #include "StageMonitorModule.h"
 
+#include "Engine/Engine.h"
 #include "HAL/IConsoleManager.h"
 #include "Misc/CoreDelegates.h"
 #include "StageMonitor.h"
 #include "StageMonitoringSettings.h"
 #include "StageMonitorSessionManager.h"
 #include "VPSettings.h"
+#include "VPRoles/Public/VPRolesSubsystem.h"
 
 const FName IStageMonitorModule::ModuleName = TEXT("StageMonitor");
 
@@ -42,17 +44,19 @@ void FStageMonitorModule::OnEngineLoopInitComplete()
 	StageMonitor = MakeUnique<FStageMonitor>();
 	StageMonitor->Initialize();
 
+	const UVirtualProductionRolesSubsystem* RolesSubsystem = GEngine->GetEngineSubsystem<UVirtualProductionRolesSubsystem>();
+
 	const UStageMonitoringSettings* Settings = GetDefault<UStageMonitoringSettings>();
 	if (Settings->MonitorSettings.ShouldAutoStartOnLaunch())
 	{
-		if (!Settings->MonitorSettings.bUseRoleFiltering || GetDefault<UVPSettings>()->GetRoles().HasAny(Settings->MonitorSettings.SupportedRoles))
+		if (!Settings->MonitorSettings.bUseRoleFiltering || RolesSubsystem->GetRolesContainer_Private().HasAny(Settings->MonitorSettings.SupportedRoles))
 		{
 			EnableMonitor(true);
 		}
 		else
 		{
 			UE_LOG(LogStageMonitor, Log, TEXT("Can't autostart StageMonitor. Role filtering is enabled and our roles (%s) are filtered out (%s)")
-				, *GetDefault<UVPSettings>()->GetRoles().ToStringSimple()
+				, *RolesSubsystem->GetActiveRolesString()
 				, *Settings->MonitorSettings.SupportedRoles.ToStringSimple())
 		}
 	}
diff --git a/Engine/Plugins/VirtualProduction/StageMonitoring/Source/StageMonitor/StageMonitor.build.cs b/Engine/Plugins/VirtualProduction/StageMonitoring/Source/StageMonitor/StageMonitor.build.cs
--- a/Engine/Plugins/VirtualProduction/StageMonitoring/Source/StageMonitor/StageMonitor.build.cs
+++ b/Engine/Plugins/VirtualProduction/StageMonitoring/Source/StageMonitor/StageMonitor.build.cs
@@ -20,12 +20,15 @@ public class StageMonitor : ModuleRules
 			new string[]
 			{
 				"CoreUObject",
+				"Engine",
 				"GameplayTags",
 				"Json",
 				"JsonUtilities",
 				"Serialization",
 				"StageMonitorCommon",
-				"VPUtilities"
+				"VPRoles",
+				"VPUtilities",
+				
 			}
 		);
 	}
diff --git a/Engine/Plugins/VirtualProduction/StageMonitoring/Source/StageMonitorCommon/Private/StageMonitorUtils.cpp b/Engine/Plugins/VirtualProduction/StageMonitoring/Source/StageMonitorCommon/Private/StageMonitorUtils.cpp
--- a/Engine/Plugins/VirtualProduction/StageMonitoring/Source/StageMonitorCommon/Private/StageMonitorUtils.cpp
+++ b/Engine/Plugins/VirtualProduction/StageMonitoring/Source/StageMonitorCommon/Private/StageMonitorUtils.cpp
@@ -2,9 +2,10 @@
 
 #include "StageMonitorUtils.h"
 
-#include "Misc/App.h"
+#include "Engine/Engine.h"
 #include "StageMonitoringSettings.h"
 #include "VPSettings.h"
+#include "VPRoles/Public/VPRolesSubsystem.h"
 
 
 namespace StageMonitorUtils
@@ -19,7 +20,10 @@ namespace StageMonitorUtils
 		Descriptor.MachineName = CachedComputerName;
 		Descriptor.ProcessId = FPlatformProcess::GetCurrentProcessId();
 
-		Descriptor.RolesStringified = GetDefault<UVPSettings>()->GetRoles().ToStringSimple();
+		if (GEngine)
+		{
+			Descriptor.RolesStringified = GEngine->GetEngineSubsystem<UVirtualProductionRolesSubsystem>()->GetActiveRolesString();
+		}
 
 		Descriptor.FriendlyName = GetDefault<UStageMonitoringSettings>()->CommandLineFriendlyName;
 		if (Descriptor.FriendlyName == NAME_None)
diff --git a/Engine/Plugins/VirtualProduction/StageMonitoring/Source/StageMonitorCommon/StageMonitorCommon.build.cs b/Engine/Plugins/VirtualProduction/StageMonitoring/Source/StageMonitorCommon/StageMonitorCommon.build.cs
--- a/Engine/Plugins/VirtualProduction/StageMonitoring/Source/StageMonitorCommon/StageMonitorCommon.build.cs
+++ b/Engine/Plugins/VirtualProduction/StageMonitoring/Source/StageMonitorCommon/StageMonitorCommon.build.cs
@@ -23,7 +23,8 @@ public class StageMonitorCommon : ModuleRules
 				"CoreUObject",
 				"Engine",
 				"GameplayTags",
-				"VPUtilities"
+				"VPRoles",
+				"VPUtilities",
 			}
 		);
 	}
