Add ConsoleVariables the VP utilities so VPStageCVars.uasset can find the UConsoleVariablesAsset.

#rnx

diff --git a/Engine/Plugins/Experimental/VirtualProductionUtilities/VirtualProductionUtilities.uplugin b/Engine/Plugins/Experimental/VirtualProductionUtilities/VirtualProductionUtilities.uplugin
--- a/Engine/Plugins/Experimental/VirtualProductionUtilities/VirtualProductionUtilities.uplugin
+++ b/Engine/Plugins/Experimental/VirtualProductionUtilities/VirtualProductionUtilities.uplugin
@@ -78,6 +78,10 @@
 		{
 			"Name": "GameplayTagsEditor",
 			"Enabled": true
+		},
+		{
+			"Name": "ConsoleVariables",
+			"Enabled": true			
 		}
 
 	]
