Enable gameplay tags config import by default when using the vp utilities plugin

diff --git a/Engine/Plugins/Experimental/VirtualProductionUtilities/Source/VPRolesEditor/Private/VPRolesEditorModule.cpp b/Engine/Plugins/Experimental/VirtualProductionUtilities/Source/VPRolesEditor/Private/VPRolesEditorModule.cpp
--- a/Engine/Plugins/Experimental/VirtualProductionUtilities/Source/VPRolesEditor/Private/VPRolesEditorModule.cpp
+++ b/Engine/Plugins/Experimental/VirtualProductionUtilities/Source/VPRolesEditor/Private/VPRolesEditorModule.cpp
@@ -4,6 +4,7 @@
 
 #include "GameplayTagContainer.h"
 #include "GameplayTagsEditorModule.h"
+#include "GameplayTagsSettings.h"
 #include "LevelEditor.h"
 #include "HAL/PlatformFileManager.h"
 #include "Misc/MessageDialog.h"
@@ -67,6 +68,11 @@ void FVPRolesEditorModule::ExtendLevelEditorToolbar()
 {
 	if (UVirtualProductionRolesSubsystem* RolesSubsystem = GEngine->GetEngineSubsystem<UVirtualProductionRolesSubsystem>())
 	{
+		if (UGameplayTagsSettings* GameplayTagsSettings = GetMutableDefault<UGameplayTagsSettings>())
+		{
+			GameplayTagsSettings->ImportTagsFromConfig = true;
+		}
+		
 		UToolMenu* Menu = UToolMenus::Get()->ExtendMenu("LevelEditor.LevelEditorToolBar.User");
 		
 		FToolMenuEntry VPRolesEntry = FToolMenuEntry::InitComboButton(
