VPBookmark - nativized VPBookmark actor

replaced BPI_Interaction with native version
replaced references to BPI_Interaction in uassets

diff --git a/Engine/Plugins/Experimental/VirtualProductionUtilities/Source/VPUtilities/Private/VPBookmarkActor.cpp b/Engine/Plugins/Experimental/VirtualProductionUtilities/Source/VPUtilities/Private/VPBookmarkActor.cpp
new file mode 100644
--- /dev/null
+++ b/Engine/Plugins/Experimental/VirtualProductionUtilities/Source/VPUtilities/Private/VPBookmarkActor.cpp
@@ -0,0 +1,204 @@
+// Copyright Epic Games, Inc. All Rights Reserved.
+
+
+#include "VPBookmarkActor.h"
+#include "UObject/ConstructorHelpers.h"
+#include "Engine/StaticMesh.h"
+#include "Components/StaticMeshComponent.h"
+#include "Components/SplineMeshComponent.h"
+#include "Components/TextRenderComponent.h"
+#include "VPUtilitiesModule.h"
+#include "VPBookmarkBlueprintLibrary.h"
+#include "VPBlueprintLibrary.h"
+#include "VPBookmark.h"
+#include "Materials/MaterialInstanceDynamic.h"
+#include "Materials/Material.h"
+
+
+AVPBookmarkActor::AVPBookmarkActor(const FObjectInitializer& ObjectInitializer)
+	:Super(ObjectInitializer)
+	{
+	PrimaryActorTick.bCanEverTick = true;
+	PrimaryActorTick.bStartWithTickEnabled = true;
+
+	//Root Component
+
+	BookmarkMesh = CreateDefaultSubobject<UStaticMeshComponent>(TEXT("BookmarkMesh"));
+	SetRootComponent(BookmarkMesh);
+
+	static ConstructorHelpers::FObjectFinder<UStaticMesh> BookmarkStaticMesh(TEXT("/VirtualProductionUtilities/Bookmark/BookmarkVPMesh.BookmarkVPMesh"));
+	if (BookmarkStaticMesh.Succeeded())
+	{
+		FoundMesh = BookmarkStaticMesh.Object;
+	}
+
+	BookmarkMesh->SetStaticMesh(FoundMesh);
+
+	static ConstructorHelpers::FObjectFinder<UMaterialInterface> BookmarkMaterialInstance(TEXT("/VirtualProductionUtilities/Bookmark/MI_BookmarkVPNoUser.MI_BookmarkVPNoUser"));
+	if (BookmarkStaticMesh.Succeeded())
+	{
+		BookmarkMaterial = BookmarkMaterialInstance.Object;
+	}
+		
+	BookmarkColor = FLinearColor(0.817708f, 0.107659f, 0.230336f);
+
+	//SplineMesh
+	SplineMesh = CreateDefaultSubobject<USplineMeshComponent>(TEXT("SplineMesh"));
+	SplineMesh->SetMobility(EComponentMobility::Movable);
+	SplineMesh->SetupAttachment(BookmarkMesh);
+	SplineMesh->SetVisibility(false);
+
+	static ConstructorHelpers::FObjectFinder<UStaticMesh> SplineStaticMesh(TEXT("/Engine/VREditor/FloatingText/LineSegmentCylinder.LineSegmentCylinder"));
+	if (SplineStaticMesh.Succeeded())
+	{
+		FoundMesh = SplineStaticMesh.Object;
+	};
+	SplineMesh->SetStaticMesh(FoundMesh);
+
+	static ConstructorHelpers::FObjectFinder<UMaterialInterface> LaserMaterial(TEXT("/Engine/VREditor/LaserPointer/LaserPointerMaterialInst.LaserPointerMaterialInst"));
+	if (LaserMaterial.Succeeded())
+	{
+		SplineMaterial = LaserMaterial.Object;
+	}
+	SplineMesh->SetMaterial(0, SplineMaterial);
+	
+
+	//Text Render
+	NameTextRender = CreateDefaultSubobject<UTextRenderComponent>(TEXT("NameTextRender"));
+	NameTextRender->SetupAttachment(BookmarkMesh);
+	NameTextRender->SetWorldSize(36);
+	NameTextRender->AddRelativeLocation(FVector(0.f, 0.f, 70.f));
+	NameTextRender->HorizontalAlignment = EHorizTextAligment(EHTA_Center);
+
+	static ConstructorHelpers::FObjectFinder<UMaterial> FoundMaterial(TEXT("/VirtualProductionUtilities/Bookmark/CameraFacingTextMaterial.CameraFacingTextMaterial"));
+	if (FoundMaterial.Succeeded())
+	{
+		TextMaterial = FoundMaterial.Object;
+	}
+	NameTextRender->SetMaterial(0, TextMaterial);
+
+}
+
+void AVPBookmarkActor::Tick(float DeltaSeconds)
+	{
+	Super::Tick(DeltaSeconds);
+
+
+	NameTextRender->SetWorldRotation(FRotator(0.f, 0.f, 0.f));
+
+	BookmarkRotation = GetActorRotation();
+	BookmarkMesh->SetWorldRotation(FRotator(0.f,BookmarkRotation.Yaw,0.f));
+
+
+#if WITH_EDITOR
+	if (GIsEditor)
+	{
+	FEditorScriptExecutionGuard ScriptGuard;
+	EditorTick(DeltaSeconds);
+	}
+#endif
+}
+
+//VP Interaction InterfaceEvents
+
+void AVPBookmarkActor::OnBookmarkActivation_Implementation(UVPBookmark* BookmarkOut, bool bIsActive)
+{
+	UE_LOG(LogVPUtilities, Display, TEXT("Bookmark Created"));
+}
+
+
+void AVPBookmarkActor::OnBookmarkChanged_Implementation(UVPBookmark* BookmarkOut)
+{
+	AActor* BookmarkActor = BookmarkOut->GetAssociatedBookmarkActor();
+	IVPBookmarkProvider* BookmarkInterface = Cast<IVPBookmarkProvider>(BookmarkActor);
+	if (BookmarkInterface)
+	{
+		BookmarkInterface->Execute_GenerateBookmarkName(BookmarkActor);
+	}
+	BookmarkObject = BookmarkOut;
+	UE_LOG(LogVPUtilities, Display, TEXT("Bookmark Updated"));
+}
+
+
+void AVPBookmarkActor::UpdateBookmarkSplineMeshIndicator_Implementation()
+{
+	UVPBlueprintLibrary::VPBookmarkSplineMeshIndicatorSetStartAndEnd(SplineMesh);
+}
+
+
+void AVPBookmarkActor::HideBookmarkSplineMeshIndicator_Implementation()
+{
+	UVPBlueprintLibrary::VPBookmarkSplineMeshIndicatorDisable(SplineMesh);
+}
+
+
+void AVPBookmarkActor::GenerateBookmarkName_Implementation()
+{
+	FString GeneratedNumber;
+	FString GeneratedLetter;
+	UVPBookmarkBlueprintLibrary::CreateVPBookmarkName(this, FString("Bookmark %n"), GeneratedNumber, GeneratedLetter);
+		
+	NameTextRender->SetText(FText::AsCultureInvariant(GeneratedNumber));
+
+}
+
+
+//VP Interaction Interface Events
+
+void AVPBookmarkActor::OnActorDroppedFromCarry_Implementation()
+{
+	UE_LOG(LogVPUtilities, Display, TEXT("Bookmark %s dropped from carry by VR Interactor"), *this->GetName());
+}
+
+void AVPBookmarkActor::OnActorSelectedForTransform_Implementation()
+{
+	UE_LOG(LogVPUtilities, Display, TEXT("Bookmark %s selected by VR Interactor"), *this->GetName());
+}
+
+void AVPBookmarkActor::OnActorDroppedFromTransform_Implementation()
+{
+	UE_LOG(LogVPUtilities, Display, TEXT("Bookmark %s dropped from transform dragging by VR Interactor"), *this->GetName());
+}
+
+
+void AVPBookmarkActor::UpdateBookmarkColor(FLinearColor Color)
+{
+
+	UMaterialInterface* Material = BookmarkMesh->GetMaterial(0);
+
+	if (Material && !Material->IsA<UMaterialInstanceDynamic>())
+	{
+		UMaterialInstanceDynamic* BookmarkMaterialInstance = UMaterialInstanceDynamic::Create(Material, BookmarkMaterial, TEXT("BookmarkMaterial"));
+
+		BookmarkMaterialInstance->ClearParameterValues();
+		BookmarkMaterialInstance->SetVectorParameterValue(TEXT("UserColor"), Color);
+
+		for (int32 i = 0; i < BookmarkMesh->GetNumMaterials(); i++)
+		{
+			BookmarkMesh->SetMaterial(i, BookmarkMaterialInstance);
+		}
+	}
+	else //If DMIs are already setup set the color value
+	{
+		DynamicMaterial = Cast<UMaterialInstanceDynamic>(Material);
+		DynamicMaterial->SetVectorParameterValue(TEXT("UserColor"), Color);
+	}
+}
+#if WITH_EDITOR
+void AVPBookmarkActor::PostEditChangeProperty(FPropertyChangedEvent& PropertyChangedEvent)
+{
+	Super::PostEditChangeProperty(PropertyChangedEvent);
+
+	if (PropertyChangedEvent.Property->GetFName() == GET_MEMBER_NAME_CHECKED(AVPBookmarkActor, BookmarkColor))
+	{
+		AVPBookmarkActor::UpdateBookmarkColor(BookmarkColor);
+	}
+}
+#endif
+
+void AVPBookmarkActor::OnConstruction(const FTransform& Transform)
+{
+	Super::OnConstruction(Transform);
+
+	AVPBookmarkActor::UpdateBookmarkColor(BookmarkColor);
+}
\ No newline at end of file
diff --git a/Engine/Plugins/Experimental/VirtualProductionUtilities/Source/VPUtilities/Public/IVPInteraction.h b/Engine/Plugins/Experimental/VirtualProductionUtilities/Source/VPUtilities/Public/IVPInteraction.h
new file mode 100644
--- /dev/null
+++ b/Engine/Plugins/Experimental/VirtualProductionUtilities/Source/VPUtilities/Public/IVPInteraction.h
@@ -0,0 +1,36 @@
+// Copyright Epic Games, Inc. All Rights Reserved.
+
+#pragma once
+
+#include "CoreMinimal.h"
+#include "UObject/Interface.h"
+#include "IVPInteraction.generated.h"
+
+class UVirtualScoutingInteraction;
+
+// This class does not need to be modified.
+UINTERFACE(BlueprintType)
+class UVPInteraction : public UInterface
+{
+	GENERATED_BODY()
+};
+
+class IVPInteraction : public IInterface  
+{
+	GENERATED_BODY()
+
+	// Add interface functions to this class. This is the class that will be inherited to implement this interface.
+public:
+
+	/** Called when actor is dropped from Carry state by the VREd interactor */
+	UFUNCTION(BlueprintNativeEvent, CallInEditor, BlueprintCallable, Category = "VirtualScouting|Interaction")
+	void OnActorDroppedFromCarry();
+
+	/** Called when an actor is selected by the VREd interactor */
+	UFUNCTION(BlueprintNativeEvent, CallInEditor, BlueprintCallable, Category = "VirtualScouting|Interaction")
+	void OnActorSelectedForTransform();
+
+	/** Called when dropped by VRTool */
+	UFUNCTION(BlueprintNativeEvent, CallInEditor, BlueprintCallable, Category = "VirtualScouting|Interaction")
+	void OnActorDroppedFromTransform();
+};
\ No newline at end of file
diff --git a/Engine/Plugins/Experimental/VirtualProductionUtilities/Source/VPUtilities/Public/VPBookmarkActor.h b/Engine/Plugins/Experimental/VirtualProductionUtilities/Source/VPUtilities/Public/VPBookmarkActor.h
new file mode 100644
--- /dev/null
+++ b/Engine/Plugins/Experimental/VirtualProductionUtilities/Source/VPUtilities/Public/VPBookmarkActor.h
@@ -0,0 +1,116 @@
+// Copyright Epic Games, Inc. All Rights Reserved.
+
+#pragma once
+
+#include "CoreMinimal.h"
+#include "VPViewportTickableActorBase.h"
+#include "UObject/UObjectGlobals.h"
+#include "VPBookmark.h"
+#include "IVPBookmarkProvider.h"
+#include "IVPInteraction.h"
+#include "VPBookmarkActor.generated.h"
+
+class UStaticMeshComponent;
+class UTextRenderComponent;
+class USplineMeshComponent;
+class UMaterialInstanceDynamic;
+class UMaterial;
+
+
+UCLASS(HideCategories=(Rendering, Lighting, HLOD, Mobile, Navigation, RayTracing, TextureStreaming), CollapseCategories)
+class VPUTILITIES_API AVPBookmarkActor : public AVPViewportTickableActorBase, public IVPInteraction, public IVPBookmarkProvider
+{
+	GENERATED_BODY()
+
+public:
+
+	AVPBookmarkActor(const FObjectInitializer& ObjectInitializer);
+
+	/** Mesh Representation in the world */
+	UPROPERTY(BlueprintReadOnly, EditDefaultsOnly, Category = "Components")
+	TObjectPtr<UStaticMeshComponent> BookmarkMesh;
+
+	/** Textrender to display bookmark name */
+	UPROPERTY(BlueprintReadOnly, EditDefaultsOnly, Category = "Components")
+	TObjectPtr<UTextRenderComponent> NameTextRender;
+
+	/** Splinemesh */
+	UPROPERTY(BlueprintReadOnly, EditDefaultsOnly, Category = "Components")
+	TObjectPtr<USplineMeshComponent> SplineMesh;
+
+	/** Color of Bookmark in MU Session */
+	UPROPERTY(BlueprintReadWrite, EditAnywhere, Category = "Bookmark", meta = (MultiLine = "true", ExposeOnSpawn = "true"))
+	FLinearColor BookmarkColor;
+
+	/** Reference to Editor Bookmark UObject*/
+	UPROPERTY(BlueprintReadWrite, EditDefaultsOnly, Category = "Bookmark", meta = (MultiLine = "true"))
+	TObjectPtr<UVPBookmark> BookmarkObject;
+
+	/** Bool to determine if this bookmark should be designated the home location */
+	UPROPERTY(BlueprintReadWrite, EditDefaultsOnly, Category = "Bookmark", meta = (MultiLine = "true"))
+	bool IsHome;
+
+	/** Texture reference to store render of viewpoint  */
+	UPROPERTY(BlueprintReadWrite, Category = "Snapshot")
+	TObjectPtr<UTexture2D> SnapshotTexture;
+
+	/**Update the mesh color and BookmarkColor variable. Intended for use with multiuser initialisation*/
+	UFUNCTION(BlueprintCallable, CallInEditor, Category = "Default")
+	void UpdateBookmarkColor(FLinearColor Color);
+
+	UPROPERTY(BlueprintReadWrite, Category = "Default")
+	TObjectPtr<UMaterialInstanceDynamic> DynamicMaterial;
+
+	UPROPERTY(EditDefaultsOnly, BlueprintReadWrite, Category = "Default")
+	TObjectPtr<UMaterialInterface> BookmarkMaterial;
+
+
+	//VPBookmark Interface events
+
+	UFUNCTION(BlueprintCallable, CallInEditor, Category = "Bookmarks")
+	void OnBookmarkActivation_Implementation(UVPBookmark* Bookmark, bool bActivate);
+
+	UFUNCTION(BlueprintCallable, CallInEditor, Category = "Bookmarks")
+	void OnBookmarkChanged_Implementation(UVPBookmark* Bookmark);
+
+	UFUNCTION(BlueprintCallable, CallInEditor, Category = "Bookmarks")
+	void UpdateBookmarkSplineMeshIndicator_Implementation();
+
+	UFUNCTION(BlueprintCallable, CallInEditor, Category = "Bookmarks")
+	void HideBookmarkSplineMeshIndicator_Implementation();
+
+	UFUNCTION(BlueprintCallable, CallInEditor, Category = "Bookmarks")
+	void GenerateBookmarkName_Implementation();
+
+
+	//VPInteraction Interface Events
+
+	UFUNCTION(BlueprintCallable, CallInEditor, Category = "Interaction")
+	void OnActorDroppedFromCarry_Implementation();
+
+	UFUNCTION(BlueprintCallable, CallInEditor, Category = "Interaction")
+	void OnActorSelectedForTransform_Implementation();
+
+	UFUNCTION(BlueprintCallable, CallInEditor, Category = "Interaction")
+	void OnActorDroppedFromTransform_Implementation();
+
+	//Overrides
+
+	virtual void OnConstruction(const FTransform& Transform) override;
+
+	virtual void Tick(float DeltaSeconds) override;
+
+#if WITH_EDITOR
+	virtual void PostEditChangeProperty(FPropertyChangedEvent& PropertyChangedEvent) override;
+#endif
+	
+private:
+
+	UMaterial* TextMaterial;
+
+	UStaticMesh* FoundMesh;
+
+	UMaterialInterface* SplineMaterial;
+
+	FRotator BookmarkRotation;
+};
