VPSpline initial checkin

diff --git a/Engine/Plugins/Experimental/VirtualProductionUtilities/Config/BaseVirtualProductionUtilities.ini b/Engine/Plugins/Experimental/VirtualProductionUtilities/Config/BaseVirtualProductionUtilities.ini
--- a/Engine/Plugins/Experimental/VirtualProductionUtilities/Config/BaseVirtualProductionUtilities.ini
+++ b/Engine/Plugins/Experimental/VirtualProductionUtilities/Config/BaseVirtualProductionUtilities.ini
@@ -8,4 +8,5 @@
 ScoutingSubsystemEditorUtilityClassPath=/VirtualProductionUtilities/BP_VPScoutingSubsystemHelpers.BP_VPScoutingSubsystemHelpers_C
 GestureManagerEditorUtilityClassPath=/VirtualProductionUtilities/VRTools/BP_GestureManager.BP_GestureManager_C
 +AdditionnalClassToLoad=/VirtualProductionUtilities/VR/VirtualScoutingInteractor.VirtualScoutingInteractor
-+AdditionnalClassToLoad=/VirtualProductionUtilities/VR/VirtualScoutingTeleporter.VirtualScoutingTeleporter
\ No newline at end of file
++AdditionnalClassToLoad=/VirtualProductionUtilities/VR/VirtualScoutingTeleporter.VirtualScoutingTeleporter
+VPSplinePreviewMeshPath=/VirtualProductionUtilities/Editor/Meshes/SM_VPSplineArraw.SM_VPSplineArraw
\ No newline at end of file
diff --git a/Engine/Plugins/Experimental/VirtualProductionUtilities/Source/VPSpline/Private/VPSpline.cpp b/Engine/Plugins/Experimental/VirtualProductionUtilities/Source/VPSpline/Private/VPSpline.cpp
new file mode 100644
--- /dev/null
+++ b/Engine/Plugins/Experimental/VirtualProductionUtilities/Source/VPSpline/Private/VPSpline.cpp
@@ -0,0 +1,23 @@
+// Copyright Epic Games, Inc. All Rights Reserved.
+
+#include "VPSpline.h"
+#include "VPSplineLog.h"
+
+#define LOCTEXT_NAMESPACE "FVPSplineModule"
+
+DEFINE_LOG_CATEGORY(LogVPSpline);
+
+void FVPSplineModule::StartupModule()
+{
+	// This code will execute after your module is loaded into memory; the exact timing is specified in the .uplugin file per-module
+}
+
+void FVPSplineModule::ShutdownModule()
+{
+	// This function may be called during shutdown to clean up your module.  For modules that support dynamic reloading,
+	// we call this function before unloading the module.
+}
+
+#undef LOCTEXT_NAMESPACE
+	
+IMPLEMENT_MODULE(FVPSplineModule, VPSpline)
\ No newline at end of file
diff --git a/Engine/Plugins/Experimental/VirtualProductionUtilities/Source/VPSpline/Private/VPSplineActor.cpp b/Engine/Plugins/Experimental/VirtualProductionUtilities/Source/VPSpline/Private/VPSplineActor.cpp
new file mode 100644
--- /dev/null
+++ b/Engine/Plugins/Experimental/VirtualProductionUtilities/Source/VPSpline/Private/VPSplineActor.cpp
@@ -0,0 +1,292 @@
+// Copyright Epic Games, Inc. All Rights Reserved.
+
+#include "VPSplineActor.h"
+#include "VPSplinePointData.h"
+#include "VPSplineLog.h"
+
+#include "CineCameraComponent.h"
+#include "UObject/ConstructorHelpers.h"
+
+#if WITH_EDITOR
+#include "ISequencer.h"
+#include "ISequencerModule.h"
+#include "LevelSequence.h"
+#include "LevelSequenceEditorBlueprintLibrary.h"
+#include "Tracks/MovieSceneFloatTrack.h"
+#include "Sections/MovieSceneFloatSection.h"
+#include "Channels/MovieSceneChannelProxy.h"
+#include "Misc/Guid.h"
+#include "Modules/ModuleManager.h"
+#include "MovieSceneToolHelpers.h"
+#include "VPUtilitiesEditorSettings.h"
+#endif
+
+AVPSplineActor::AVPSplineActor(const FObjectInitializer& ObjectInitializer)
+	: Super(ObjectInitializer)
+{
+	PrimaryActorTick.bCanEverTick = true;
+	PrimaryActorTick.bStartWithTickEnabled = true;
+
+	SplineComp = CreateDefaultSubobject<UVPSplineComponent>("VPSpline");
+	SetRootComponent(SplineComp);
+
+	SplineAttachment = CreateDefaultSubobject<USceneComponent>("Spline Attachment");
+	SplineAttachment->SetupAttachment(SplineComp);
+
+#if WITH_EDITORONLY_DATA
+	PreviewMeshScale = 1.f;
+
+	if (!IsRunningCommandlet())
+	{
+		FSoftObjectPath PreviewMeshPath = GetDefault<UVPUtilitiesEditorSettings>()->VPSplinePreviewMeshPath;
+		static ConstructorHelpers::FObjectFinder<UStaticMesh> MeshFinder(*PreviewMeshPath.ToString());
+		if (MeshFinder.Succeeded())
+		{
+			PreviewMesh = MeshFinder.Object;
+		}
+		else
+		{
+			UE_LOG(LogVPSpline, Warning, TEXT("Failed to find spline preview mesh: %s"), *PreviewMeshPath.ToString());
+		}
+	}
+
+	PreviewMeshComp = CreateOptionalDefaultSubobject<UStaticMeshComponent>(TEXT("PreviewMesh"));
+	if (PreviewMeshComp)
+	{
+		PreviewMeshComp->SetIsVisualizationComponent(true);
+		PreviewMeshComp->SetCollisionProfileName(UCollisionProfile::NoCollision_ProfileName);
+		PreviewMeshComp->bHiddenInGame = true;
+		PreviewMeshComp->CastShadow = false;
+		PreviewMeshComp->SetupAttachment(SplineAttachment);
+	}
+#endif
+
+}
+
+void AVPSplineActor::SetPointFromActor(const AActor* Actor)
+{
+	if (Actor != nullptr)
+	{
+		FVPSplinePointData Data;
+		if (const ACineCameraActor* CineCameraActor = Cast<ACineCameraActor>(Actor))
+		{
+			if (UCineCameraComponent* CameraComp = CineCameraActor->GetCineCameraComponent())
+			{
+				Data.Location = CameraComp->GetComponentLocation();
+				Data.Rotation = CameraComp->GetComponentRotation();
+				Data.FocalLength = CameraComp->CurrentFocalLength;
+				Data.Aperture = CameraComp->CurrentAperture;
+				Data.FocusDistance = CameraComp->CurrentFocusDistance;
+			}
+		}
+		else
+		{
+			Data.Location = Actor->GetActorLocation();
+			Data.Rotation = Actor->GetActorRotation();
+		}
+
+		SetPointByValue(Data);
+	}
+}
+
+void AVPSplineActor::SetPointByValue(const FVPSplinePointData& Data)
+{
+	int32 Index;
+	if (SplineComp->FindSplineDataAtPosition(CurrentPosition, Index))
+	{
+		UE_LOG(LogVPSpline, Display, TEXT("Updating keyframe at %d"), Index);
+		SplineComp->UpdateSplineDataAtIndex(Index, Data);
+	}
+	else
+	{
+		UE_LOG(LogVPSpline, Display, TEXT("Adding keyframe at %f"), CurrentPosition);
+		SplineComp->AddSplineDataAtPosition(CurrentPosition, Data);
+	}
+}
+
+void AVPSplineActor::RemoveCurrentPoint()
+{
+	int32 Index;
+	if (SplineComp->FindSplineDataAtPosition(CurrentPosition, Index))
+	{
+		UE_LOG(LogVPSpline, Display, TEXT("Removing keyframe at %d"), Index);
+		SplineComp->RemoveSplinePoint(Index);
+	}
+}
+
+
+void AVPSplineActor::GoToNextPosition()
+{
+	UVPSplineMetadata* Metadata = SplineComp ? Cast<UVPSplineMetadata>(SplineComp->GetSplinePointsMetadata()) : nullptr;
+	if (Metadata != nullptr)
+	{
+		for (int32 i = 0; i < Metadata->NormalizedPosition.Points.Num(); ++i)
+		{
+			if (CurrentPosition < Metadata->NormalizedPosition.Points[i].OutVal)
+			{
+				CurrentPosition = Metadata->NormalizedPosition.Points[i].OutVal;
+				break;
+			}
+		}
+	}
+}
+
+void AVPSplineActor::GoToPrevPosition()
+{
+	UVPSplineMetadata* Metadata = SplineComp ? Cast<UVPSplineMetadata>(SplineComp->GetSplinePointsMetadata()) : nullptr;
+	if (Metadata != nullptr)
+	{
+		for (int32 i = Metadata->NormalizedPosition.Points.Num(); i > 0; --i)
+		{
+			if (CurrentPosition > Metadata->NormalizedPosition.Points[i-1].OutVal)
+			{
+				CurrentPosition = Metadata->NormalizedPosition.Points[i-1].OutVal;
+				break;
+			}
+		}
+	}
+}
+
+
+
+void AVPSplineActor::BakePointsToSequence()
+{
+#if WITH_EDITOR
+	const FName PropName = FName(TEXT("CurrentPosition"));
+	const FString PropPath = PropName.ToString();
+
+	ULevelSequence* LevelSequence = ULevelSequenceEditorBlueprintLibrary::GetCurrentLevelSequence();
+	if (LevelSequence != nullptr)
+	{
+		UMovieScene* MovieScene = LevelSequence->GetMovieScene();
+		if (MovieScene != nullptr)
+		{
+			FGuid ActorBinding;
+			ActorBinding = LevelSequence->FindBindingFromObject(this, GetWorld());
+			if (!ActorBinding.IsValid())
+			{
+				ActorBinding = MovieScene->AddPossessable(GetActorLabel(), GetClass());
+				LevelSequence->BindPossessableObject(ActorBinding, *this, GetWorld());
+				MovieScene->Modify();
+			}
+
+			UMovieSceneFloatTrack* ParamTrack = MovieScene->FindTrack<UMovieSceneFloatTrack>(ActorBinding, PropName);
+			if (ParamTrack == nullptr)
+			{
+				ParamTrack = MovieScene->AddTrack<UMovieSceneFloatTrack>(ActorBinding);
+				if (!ParamTrack)
+				{
+					UE_LOG(LogVPSpline, Error, TEXT("%s: Failed to add float track"), *GetActorNameOrLabel());
+					return;
+				}
+				ParamTrack->RemoveAllAnimationData();
+				ParamTrack->SetPropertyNameAndPath(PropName, PropPath);
+				ParamTrack->Modify();
+				UE_LOG(LogVPSpline, Log, TEXT("%s: %s track added"), *GetActorNameOrLabel(), *PropPath);
+			}
+
+			bool bSectionAdded = false;
+			UMovieSceneFloatSection* Section = Cast<UMovieSceneFloatSection>(ParamTrack->FindOrAddSection(0, bSectionAdded));
+			Section->SetRange(TRange<FFrameNumber>::All());
+			FMovieSceneFloatChannel* Channel = Section->GetChannelProxy().GetChannel<FMovieSceneFloatChannel>(0);
+			ensure(Channel);
+			Channel->Reset();
+			Channel->SetDefault(0.0f);
+			TMovieSceneChannelData<FMovieSceneFloatValue> ChannelData = Channel->GetData();
+			UVPSplineMetadata* Metadata = SplineComp ? Cast<UVPSplineMetadata>(SplineComp->GetSplinePointsMetadata()) : nullptr;
+			
+			if (Metadata != nullptr)
+			{
+				TRange<FFrameNumber> PlaybackRange = MovieScene->GetPlaybackRange();
+				int32 StartFrame = PlaybackRange.GetLowerBoundValue().Value;
+				int32 EndFrame = PlaybackRange.GetUpperBoundValue().Value;
+				int32 NumPoints = Metadata->NormalizedPosition.Points.Num();
+				for (int32 i = 0; i < NumPoints; ++i)
+				{
+					float Value = Metadata->NormalizedPosition.Points[i].OutVal;
+					int32 Frame = FMath::FloorToInt((float)(EndFrame - StartFrame) * Value) + StartFrame;
+					AddKeyToChannel(Channel, FFrameNumber(Frame), Value, EMovieSceneKeyInterpolation::Linear);
+				}
+			}
+		}
+		ULevelSequenceEditorBlueprintLibrary::RefreshCurrentLevelSequence();
+	}
+#endif
+}
+
+void AVPSplineActor::Tick(float DeltaTime)
+{
+	Super::Tick(DeltaTime);
+	UpdateSplineAttachment();
+}
+
+bool AVPSplineActor::ShouldTickIfViewportsOnly() const
+{
+	return true;
+}
+
+USceneComponent* AVPSplineActor::GetDefaultAttachComponent() const
+{
+	return SplineAttachment;
+}
+
+void AVPSplineActor::UpdateSplineAttachment()
+{
+	if(!GetWorld())
+	{
+		return;
+	}
+
+	if (SplineComp && SplineAttachment)
+	{
+		float InputKey = SplineComp->GetInputKeyAtPosition(CurrentPosition);
+		FVector const Position = SplineComp->GetLocationAtSplineInputKey(InputKey, ESplineCoordinateSpace::World);
+		FQuat const Rotation = SplineComp->GetQuaternionAtSplineInputKey(InputKey, ESplineCoordinateSpace::World);
+		SplineAttachment->SetWorldTransform(FTransform(Rotation, Position));
+	}
+
+#if WITH_EDITOR
+	if (GIsEditor)
+	{
+		const UWorld* const MyWorld = GetWorld();
+		if (MyWorld && !MyWorld->IsGameWorld())
+		{
+			UpdatePreviewMesh();
+		}
+	}
+#endif
+}
+
+#if WITH_EDITOR
+void AVPSplineActor::PostEditChangeProperty(FPropertyChangedEvent& PropertyChangedEvent)
+{
+	Super::PostEditChangeProperty(PropertyChangedEvent);
+	UpdateSplineAttachment();
+}
+void AVPSplineActor::PostEditUndo()
+{
+	Super::PostEditUndo();
+	UpdateSplineAttachment();
+}
+
+void AVPSplineActor::PostEditMove(bool bFinished)
+{
+	Super::PostEditMove(bFinished);
+	UpdateSplineAttachment();
+}
+#endif
+
+#if WITH_EDITORONLY_DATA
+void AVPSplineActor::UpdatePreviewMesh()
+{
+	if (PreviewMeshComp && PreviewMeshComp->GetStaticMesh() != PreviewMesh)
+	{
+		PreviewMeshComp->SetStaticMesh(PreviewMesh);
+	}
+
+	if (PreviewMeshComp)
+	{
+		PreviewMeshComp->SetWorldScale3D(FVector(PreviewMeshScale, PreviewMeshScale, PreviewMeshScale));
+	}
+}
+#endif
\ No newline at end of file
diff --git a/Engine/Plugins/Experimental/VirtualProductionUtilities/Source/VPSpline/Private/VPSplineComponent.cpp b/Engine/Plugins/Experimental/VirtualProductionUtilities/Source/VPSpline/Private/VPSplineComponent.cpp
new file mode 100644
--- /dev/null
+++ b/Engine/Plugins/Experimental/VirtualProductionUtilities/Source/VPSpline/Private/VPSplineComponent.cpp
@@ -0,0 +1,258 @@
+// Copyright Epic Games, Inc. All Rights Reserved.
+
+#include "VPSplineComponent.h"
+
+#include "VPSpline.h"
+#include "VPSplineLog.h"
+
+UVPSplineComponent::UVPSplineComponent(const FObjectInitializer& ObjectInitializer)
+	: Super(ObjectInitializer)
+{
+	SplineCurves.Position.Points.Reset(10);
+	SplineCurves.Rotation.Points.Reset(10);
+	SplineCurves.Scale.Points.Reset(10);
+	VPSplineMetadata = ObjectInitializer.CreateDefaultSubobject<UVPSplineMetadata>(this, TEXT("VPSplineMetadata"));
+	VPSplineMetadata->Reset(10);
+}
+
+
+USplineMetadata* UVPSplineComponent::GetSplinePointsMetadata()
+{
+	return VPSplineMetadata;
+}
+
+const USplineMetadata* UVPSplineComponent::GetSplinePointsMetadata() const
+{
+	return VPSplineMetadata;
+}
+
+void UVPSplineComponent::PostLoad()
+{
+	Super::PostLoad();
+	if (VPSplineMetadata)
+	{
+		SynchronizeProperties();
+	}
+}
+
+
+TStructOnScope<FActorComponentInstanceData> UVPSplineComponent::GetComponentInstanceData() const
+{
+	TStructOnScope<FActorComponentInstanceData> InstanceData = MakeStructOnScope<FActorComponentInstanceData, FVPSplineInstanceData>(this);
+	FVPSplineInstanceData* SplineInstanceData = InstanceData.Cast<FVPSplineInstanceData>();
+
+	if (bSplineHasBeenEdited)
+	{
+		SplineInstanceData->VPSplineMetadata = VPSplineMetadata;
+		SplineInstanceData->SplineCurves = SplineCurves;
+	}
+
+	SplineInstanceData->bSplineHasBeenEdited = bSplineHasBeenEdited;
+
+	return InstanceData;
+}
+
+void UVPSplineComponent::ApplyComponentInstanceData(FVPSplineInstanceData* SplineInstanceData, const bool bPostUCS)
+{
+	check(SplineInstanceData);
+
+	if (bPostUCS)
+	{
+		if (bInputSplinePointsToConstructionScript)
+		{
+			// Don't reapply the saved state after the UCS has run if we are inputting the points to it.
+			// This allows the UCS to work on the edited points and make its own changes.
+			return;
+		}
+		else
+		{
+			//bModifiedByConstructionScript = (SplineInstanceData->SplineCurvesPreUCS != SplineCurves);
+
+			// If we are restoring the saved state, unmark the SplineCurves property as 'modified'.
+			// We don't want to consider that these changes have been made through the UCS.
+			TArray<FProperty*> Properties;
+			//Properties.Emplace(FindFProperty<FProperty>(UVPSplineComponent::StaticClass(), GET_MEMBER_NAME_CHECKED(UVPSplineComponent, SplineCurves)));
+			Properties.Emplace(FindFProperty<FProperty>(UVPSplineComponent::StaticClass(), GET_MEMBER_NAME_CHECKED(UVPSplineComponent, SplineCurves)));
+			RemoveUCSModifiedProperties(Properties);
+		}
+	}
+	else
+	{
+		//SplineInstanceData->SplineCurvesPreUCS = SplineCurves;
+	}
+
+	if (SplineInstanceData->bSplineHasBeenEdited)
+	{
+		// Copy metadata to current component
+		if (VPSplineMetadata && SplineInstanceData->VPSplineMetadata)
+		{
+			VPSplineMetadata->Modify();
+			//VPSplineMetadata = SplineInstanceData->VPSplineMetadata;
+			UEngine::CopyPropertiesForUnrelatedObjects(SplineInstanceData->VPSplineMetadata, VPSplineMetadata);
+		}
+	
+		bModifiedByConstructionScript = false;
+	}
+
+	UpdateSpline();
+	SynchronizeProperties();
+}
+
+void UVPSplineComponent::SynchronizeProperties()
+{
+	int32 NumOfPoints = GetNumberOfSplinePoints();
+	if (VPSplineMetadata && NumOfPoints > 0)
+	{
+		VPSplineMetadata->Fixup(NumOfPoints, this);
+
+		// Fixing invalid NormalizedPosition
+		// For now, it just finds metadata item with NormalizedPosition less than 0.
+
+		int32 NumOfValidPoints = 0;
+		for (int32 Index = 0; Index < NumOfPoints; ++Index)
+		{
+			float CurrPosition = VPSplineMetadata->NormalizedPosition.Points[Index].OutVal;
+			if (CurrPosition >= -0.0001)
+			{
+				NumOfValidPoints += 1;
+			}
+		}
+		
+		if (NumOfPoints - NumOfValidPoints > 0)
+		{
+			UE_LOG(LogVPSpline, Warning, TEXT("%s: Num Of Invalid Keys: %d"), *GetReadableName(), (NumOfPoints - NumOfValidPoints));
+			float PositionIncr = 1.0 / (float)(NumOfPoints-1);
+
+			UE_LOG(LogVPSpline, Warning, TEXT("Updating NormalizedPosition metadata: (Total: %d, Valid: %d, Incr: %f"), NumOfPoints, NumOfValidPoints, PositionIncr);
+			for (int32 Index = 0; Index < NumOfPoints; ++Index)
+			{
+				VPSplineMetadata->NormalizedPosition.Points[Index].OutVal = (float)Index * PositionIncr;
+			}
+		}
+	}
+}
+
+void UVPSplineComponent::SetFocalLengthAtSplinePoint(const int32 PointIndex, const float Value)
+{
+	int32 NumPoints = VPSplineMetadata-> FocalLength.Points.Num();
+	check(PointIndex >= 0 && PointIndex < NumPoints);
+	VPSplineMetadata->Modify();
+	VPSplineMetadata->FocalLength.Points[PointIndex].OutVal = Value;
+}
+
+void UVPSplineComponent::SetApertureAtSplinePoint(const int32 PointIndex, const float Value)
+{
+	int32 NumPoints = VPSplineMetadata->Aperture.Points.Num();
+	check(PointIndex >= 0 && PointIndex < NumPoints);
+	VPSplineMetadata->Modify();
+	VPSplineMetadata->Aperture.Points[PointIndex].OutVal = Value;
+}
+
+void UVPSplineComponent::SetFocusDistanceAtSplinePoint(const int32 PointIndex, const float Value)
+{
+	int32 NumPoints = VPSplineMetadata->FocusDistance.Points.Num();
+	check(PointIndex >= 0 && PointIndex < NumPoints);
+	VPSplineMetadata->Modify();
+	VPSplineMetadata->FocusDistance.Points[PointIndex].OutVal = Value;
+}
+
+
+void UVPSplineComponent::SetNormalizedPositionAtSplinePoint(const int32 PointIndex, const float Value)
+{
+	int32 NumPoints = VPSplineMetadata->NormalizedPosition.Points.Num();
+	check(PointIndex >= 0 && PointIndex < NumPoints);
+	VPSplineMetadata->Modify();
+	VPSplineMetadata->NormalizedPosition.Points[PointIndex].OutVal = Value;
+}
+
+bool UVPSplineComponent::FindSplineDataAtPosition(const float InPosition, int32& OutIndex) const
+{
+	int32 NumPoints = VPSplineMetadata->NormalizedPosition.Points.Num();
+	for (int32 i = 0; i < NumPoints; ++i)
+	{
+		if (FMath::IsNearlyEqual(InPosition, VPSplineMetadata->NormalizedPosition.Points[i].OutVal))
+		{
+			OutIndex = i;
+			return true;
+		}
+	}
+	OutIndex = -1;
+	return false;
+}
+
+float UVPSplineComponent::GetInputKeyAtPosition(const float InPosition)
+{
+	float OutValue = 0.0f;
+	int32 NumPoints = VPSplineMetadata->NormalizedPosition.Points.Num();
+	for (int32 i = 0; i < NumPoints; ++i)
+	{
+		if (InPosition < VPSplineMetadata->NormalizedPosition.Points[i].OutVal)
+		{
+			if (i > 0)
+			{
+				float Value0 = VPSplineMetadata->NormalizedPosition.Points[i - 1].OutVal;
+				float Value1 = VPSplineMetadata->NormalizedPosition.Points[i].OutVal;
+				OutValue = (InPosition - Value0) / (Value1 - Value0) + (float)(i-1);
+			}
+			break;
+		}
+		OutValue = (float)i;
+	}
+	return OutValue;
+}
+
+void UVPSplineComponent::UpdateSplineDataAtIndex(const int InIndex, const FVPSplinePointData& InPointData)
+{
+	ESplinePointType::Type PointType = GetSplinePointType(InIndex);
+	SetLocationAtSplinePoint(InIndex, InPointData.Location, ESplineCoordinateSpace::World);
+	SetRotationAtSplinePoint(InIndex, InPointData.Rotation, ESplineCoordinateSpace::World);
+	SetSplinePointType(InIndex, PointType);
+	SetFocalLengthAtSplinePoint(InIndex, InPointData.FocalLength);
+	SetApertureAtSplinePoint(InIndex, InPointData.Aperture);
+	SetFocusDistanceAtSplinePoint(InIndex, InPointData.FocusDistance);
+}
+
+void UVPSplineComponent::AddSplineDataAtPosition(const float InPosition, const FVPSplinePointData& InPointData)
+{
+	int32 NewIndex = 0;
+	int32 NumPoints = VPSplineMetadata->NormalizedPosition.Points.Num();
+	for (int32 i = 0; i < NumPoints; ++i)
+	{
+		if (InPosition <= VPSplineMetadata->NormalizedPosition.Points[i].OutVal)
+		{
+			break;
+		}
+		NewIndex++;
+	}
+	AddSplinePointAtIndex(InPointData.Location, NewIndex, ESplineCoordinateSpace::World);
+	SetRotationAtSplinePoint(NewIndex, InPointData.Rotation, ESplineCoordinateSpace::World);
+	SetFocalLengthAtSplinePoint(NewIndex, InPointData.FocalLength);
+	SetApertureAtSplinePoint(NewIndex, InPointData.Aperture);
+	SetFocusDistanceAtSplinePoint(NewIndex, InPointData.FocusDistance);
+	SetNormalizedPositionAtSplinePoint(NewIndex, InPosition);
+	SetSplinePointType(NewIndex, ESplinePointType::Curve);
+}
+
+
+#if WITH_EDITOR
+void UVPSplineComponent::PostEditChangeProperty(FPropertyChangedEvent& PropertyChangedEvent)
+{
+	OnSplineEdited.ExecuteIfBound();
+	Super::PostEditChangeProperty(PropertyChangedEvent);
+	SynchronizeProperties();
+}
+#endif
+
+void FVPSplineInstanceData::ApplyToComponent(UActorComponent* Component, const ECacheApplyPhase CacheApplyPhase)
+{
+	if (UVPSplineComponent* SplineComp = CastChecked<UVPSplineComponent>(Component))
+	{
+		// This ensures there is no stale data causing issues where the spline is marked as read-only even though it shouldn't.
+		// There might be a better solution, but this works.
+		SplineComp->UpdateSpline();
+
+		Super::ApplyToComponent(Component, CacheApplyPhase);
+		SplineComp->ApplyComponentInstanceData(this, (CacheApplyPhase == ECacheApplyPhase::PostUserConstructionScript));
+	}
+}
+
diff --git a/Engine/Plugins/Experimental/VirtualProductionUtilities/Source/VPSpline/Private/VPSplineLog.h b/Engine/Plugins/Experimental/VirtualProductionUtilities/Source/VPSpline/Private/VPSplineLog.h
new file mode 100644
--- /dev/null
+++ b/Engine/Plugins/Experimental/VirtualProductionUtilities/Source/VPSpline/Private/VPSplineLog.h
@@ -0,0 +1,8 @@
+// Copyright Epic Games, Inc. All Rights Reserved.
+
+#pragma once
+
+#include "CoreMinimal.h"
+#include "Logging/LogMacros.h"
+
+DECLARE_LOG_CATEGORY_EXTERN(LogVPSpline, Log, All);
diff --git a/Engine/Plugins/Experimental/VirtualProductionUtilities/Source/VPSpline/Private/VPSplineMetadata.cpp b/Engine/Plugins/Experimental/VirtualProductionUtilities/Source/VPSpline/Private/VPSplineMetadata.cpp
new file mode 100644
--- /dev/null
+++ b/Engine/Plugins/Experimental/VirtualProductionUtilities/Source/VPSpline/Private/VPSplineMetadata.cpp
@@ -0,0 +1,249 @@
+// Copyright Epic Games, Inc. All Rights Reserved.
+
+#include "VPSplineMetadata.h"
+#include "VPSplineComponent.h"
+#include "VPSplineLog.h"
+
+void UVPSplineMetadata::InsertPoint(int32 Index, float t, bool bClosedLoop)
+{
+	if (Index < 0)
+	{
+		UE_LOG(LogVPSpline, Error, TEXT("InsertPoint - invalid Index: %d"), Index);
+		return;
+	}
+
+	Modify();
+
+	int32 NumPoints = FocalLength.Points.Num();
+	float InputKey = static_cast<float>(Index);
+
+	if (Index >= NumPoints)
+	{
+		// Just add point to the end instead of trying to insert
+		AddPoint(InputKey);
+	}
+	else
+	{
+		int32 PrevIndex = (bClosedLoop && Index == 0 ? NumPoints - 1 : Index - 1);
+		bool bHasPrevIndex = (PrevIndex >= 0 && PrevIndex < NumPoints);
+		float NewFocalLengthVal = FocalLength.Points[Index].OutVal;
+		float NewApertureVal = Aperture.Points[Index].OutVal;
+		float NewFocusDistanceVal = FocusDistance.Points[Index].OutVal;
+		float NewNormalizedPositionVal = NormalizedPosition.Points[Index].OutVal;
+
+		if (bHasPrevIndex)
+		{
+			float PrevFocalLengthVal = FocalLength.Points[PrevIndex].OutVal;
+			float PrevApertureVal = Aperture.Points[PrevIndex].OutVal;
+			float PrevFocusDistanceVal = FocusDistance.Points[PrevIndex].OutVal;
+			float PrevNormalizedPositionVal = NormalizedPosition.Points[PrevIndex].OutVal;
+			NewFocalLengthVal = FMath::LerpStable(PrevFocalLengthVal, NewFocalLengthVal, t);
+			NewApertureVal = FMath::LerpStable(PrevApertureVal, NewApertureVal, t);
+			NewFocusDistanceVal = FMath::LerpStable(PrevFocusDistanceVal, NewFocusDistanceVal, t);
+			NewNormalizedPositionVal = FMath::LerpStable(PrevNormalizedPositionVal, NewNormalizedPositionVal, t);
+		}
+		FInterpCurvePoint<float> NewNormalizedPosition(InputKey, NewNormalizedPositionVal);
+		NormalizedPosition.Points.Insert(NewNormalizedPosition, Index);
+
+		FInterpCurvePoint<float> NewFocalLength(InputKey, NewFocalLengthVal);
+		FocalLength.Points.Insert(NewFocalLength, Index);
+
+		FInterpCurvePoint<float> NewAperture(InputKey, NewApertureVal);
+		Aperture.Points.Insert(NewAperture, Index);
+
+		FInterpCurvePoint<float> NewFocusDistance(InputKey, NewFocusDistanceVal);
+		FocusDistance.Points.Insert(NewFocusDistance, Index);
+
+		for (int32 i = Index + 1; i < FocalLength.Points.Num(); ++i)
+		{
+			FocalLength.Points[i].InVal += 1.0f;
+			Aperture.Points[i].InVal += 1.0f;
+			FocusDistance.Points[i].InVal += 1.0f;
+			NormalizedPosition.Points[i].InVal += 1.0f;
+		}
+	}
+}
+
+void UVPSplineMetadata::UpdatePoint(int32 Index, float t, bool bClosedLoop)
+{
+	int32 NumPoints = FocalLength.Points.Num();
+	if (!FocalLength.Points.IsValidIndex(Index))
+	{
+		UE_LOG(LogVPSpline, Error, TEXT("UpdatePoint - invalid Index: %d"), Index);
+		return;
+	}
+
+	int32 PrevIndex = (bClosedLoop && Index == 0 ? NumPoints - 1 : Index - 1);
+	int32 NextIndex = (bClosedLoop && Index + 1 > NumPoints ? 0 : Index + 1);
+
+	bool bHasPrevIndex = (PrevIndex >= 0 && PrevIndex < NumPoints);
+	bool bHasNextIndex = (NextIndex >= 0 && NextIndex < NumPoints);
+
+	Modify();
+
+	if (bHasPrevIndex && bHasNextIndex)
+	{
+		float PrevNormalizedPositionVal = NormalizedPosition.Points[PrevIndex].OutVal;
+		float PrevFocalLengthVal = FocalLength.Points[PrevIndex].OutVal;
+		float PrevApertureVal = Aperture.Points[PrevIndex].OutVal;
+		float PrevFocusDistanceVal = FocusDistance.Points[PrevIndex].OutVal;
+		float NextNormalizedPositionVal = NormalizedPosition.Points[NextIndex].OutVal;
+		float NextFocalLengthVal = FocalLength.Points[NextIndex].OutVal;
+		float NextApertureVal = Aperture.Points[NextIndex].OutVal;
+		float NextFocusDistanceVal = FocusDistance.Points[NextIndex].OutVal;
+
+		NormalizedPosition.Points[Index].OutVal = FMath::LerpStable(PrevNormalizedPositionVal, NextNormalizedPositionVal, t);
+		FocalLength.Points[Index].OutVal = FMath::LerpStable(PrevFocalLengthVal, NextFocalLengthVal, t);
+		Aperture.Points[Index].OutVal = FMath::LerpStable(PrevApertureVal, NextApertureVal, t);
+		FocusDistance.Points[Index].OutVal = FMath::LerpStable(PrevFocusDistanceVal, NextFocusDistanceVal, t);
+	}
+}
+
+void UVPSplineMetadata::AddPoint(float InputKey)
+{
+	Modify();
+
+	float NewNormalizedPositionVal = -1.0f;
+	float NewFocalLengthVal = 35.0f;
+	float NewApertureVal = 2.8f;
+	float NewFocusDistanceVal = 100000.f;
+	
+	int Index = FocalLength.Points.Num() - 1;
+	if (Index >= 0)
+	{
+		NewFocalLengthVal = FocalLength.Points[Index].OutVal;
+		NewApertureVal = Aperture.Points[Index].OutVal;
+		NewFocusDistanceVal = FocusDistance.Points[Index].OutVal;
+		NewNormalizedPositionVal = NormalizedPosition.Points[Index].OutVal;
+	}
+
+	float NewInputKey = static_cast<float>(++Index);
+	FocalLength.Points.Emplace(NewInputKey, NewFocalLengthVal);
+	Aperture.Points.Emplace(NewInputKey, NewApertureVal);
+	FocusDistance.Points.Emplace(NewInputKey, NewFocusDistanceVal);
+	NormalizedPosition.Points.Emplace(NewInputKey, NewNormalizedPositionVal);
+
+}
+
+void UVPSplineMetadata::RemovePoint(int32 Index)
+{
+	check(Index < FocalLength.Points.Num());
+
+	Modify();
+	NormalizedPosition.Points.RemoveAt(Index);
+	FocalLength.Points.RemoveAt(Index);
+	Aperture.Points.RemoveAt(Index);
+	FocusDistance.Points.RemoveAt(Index);
+
+	for (int32 i = Index; i < FocalLength.Points.Num(); ++i)
+	{
+		NormalizedPosition.Points[i].InVal -= 1.0f;
+		FocalLength.Points[i].InVal -= 1.0f;
+		Aperture.Points[i].InVal -= 1.0f;
+		FocusDistance.Points[i].InVal -= 1.0f;
+	}
+}
+
+void UVPSplineMetadata::DuplicatePoint(int32 Index)
+{
+	check(Index < FocalLength.Points.Num());
+
+	float NewValue = -1.0f;
+	float CurrValue = NormalizedPosition.Points[Index].OutVal;
+
+	if (NormalizedPosition.Points.Num() > 1)
+	{
+		float PrevValue = CurrValue;
+		float NextValue = CurrValue;
+
+		if (Index + 1 < NormalizedPosition.Points.Num())
+		{
+			NewValue = (CurrValue + NormalizedPosition.Points[Index + 1].OutVal) * 0.5;
+		}
+		else if (Index + 1 == NormalizedPosition.Points.Num())
+		{
+			NewValue = CurrValue;
+			CurrValue = (CurrValue + NormalizedPosition.Points[Index - 1].OutVal) * 0.5;
+		}
+	}
+
+	Modify();
+	NormalizedPosition.Points.Insert(FInterpCurvePoint<float>(NormalizedPosition.Points[Index]), Index);
+	FocalLength.Points.Insert(FInterpCurvePoint<float>(FocalLength.Points[Index]), Index);
+	Aperture.Points.Insert(FInterpCurvePoint<float>(Aperture.Points[Index]), Index);
+	FocusDistance.Points.Insert(FInterpCurvePoint<float>(FocusDistance.Points[Index]), Index);
+
+	NormalizedPosition.Points[Index].OutVal = CurrValue;
+	NormalizedPosition.Points[Index+1].OutVal = NewValue;
+
+
+	for (int32 i = Index + 1; i < FocalLength.Points.Num(); ++i)
+	{
+		NormalizedPosition.Points[i].InVal += 1.0f;
+		FocalLength.Points[i].InVal += 1.0f;
+		Aperture.Points[i].InVal += 1.0f;
+		FocusDistance.Points[i].InVal += 1.0f;
+	}
+
+}
+
+void UVPSplineMetadata::CopyPoint(const USplineMetadata* FromSplineMetadata, int32 FromIndex, int32 ToIndex)
+{
+	check(FromSplineMetadata != nullptr);
+
+	if (const UVPSplineMetadata* FromMetadata = Cast<UVPSplineMetadata>(FromSplineMetadata))
+	{
+		check(ToIndex < FocalLength.Points.Num());
+		check(FromIndex < FromMetadata->FocalLength.Points.Num());
+
+		Modify();
+		FocalLength.Points[ToIndex].OutVal = FromMetadata->FocalLength.Points[FromIndex].OutVal;
+		Aperture.Points[ToIndex].OutVal = FromMetadata->Aperture.Points[FromIndex].OutVal;
+		FocusDistance.Points[ToIndex].OutVal = FromMetadata->FocusDistance.Points[FromIndex].OutVal;
+	}
+}
+
+void UVPSplineMetadata::Reset(int32 NumPoints)
+{
+	Modify();
+	FocalLength.Points.Reset(NumPoints);
+	Aperture.Points.Reset(NumPoints);
+	FocusDistance.Points.Reset(NumPoints);
+	NormalizedPosition.Points.Reset(NumPoints);
+}
+
+#if WITH_EDITORONLY_DATA
+template <class T>
+void FixupCurve(FInterpCurve<T>& Curve, const T& DefaultValue, int32 NumPoints)
+{
+	// Fixup bad InVal values from when the add operation below used the wrong value
+	for (int32 PointIndex = 0; PointIndex < Curve.Points.Num(); PointIndex++)
+	{
+		float InVal = PointIndex;
+		Curve.Points[PointIndex].InVal = InVal;
+	}
+
+	while (Curve.Points.Num() < NumPoints)
+	{
+		// InVal is the point index which is ascending so use previous point plus one.
+		float InVal = Curve.Points.Num() > 0 ? Curve.Points[Curve.Points.Num() - 1].InVal + 1.0f : 0.0f;
+		Curve.Points.Add(FInterpCurvePoint<T>(InVal, DefaultValue));
+	}
+
+	if (Curve.Points.Num() > NumPoints)
+	{
+		Curve.Points.RemoveAt(NumPoints, Curve.Points.Num() - NumPoints);
+	}
+}
+#endif
+
+void UVPSplineMetadata::Fixup(int32 NumPoints, USplineComponent* SplineComp)
+{
+	const FVPSplineCurveDefaults& Defaults = CastChecked<UVPSplineComponent>(SplineComp)->CameraSplineDefaults;
+#if WITH_EDITORONLY_DATA
+	FixupCurve(FocalLength, Defaults.DefaultFocalLength, NumPoints);
+	FixupCurve(Aperture, Defaults.DefaultAperture, NumPoints);
+	FixupCurve(FocusDistance, Defaults.DefaultFocusDistance, NumPoints);
+	FixupCurve(NormalizedPosition, Defaults.DefaultNormalizedPosition, NumPoints);
+#endif
+}
\ No newline at end of file
diff --git a/Engine/Plugins/Experimental/VirtualProductionUtilities/Source/VPSpline/Public/VPSpline.h b/Engine/Plugins/Experimental/VirtualProductionUtilities/Source/VPSpline/Public/VPSpline.h
new file mode 100644
--- /dev/null
+++ b/Engine/Plugins/Experimental/VirtualProductionUtilities/Source/VPSpline/Public/VPSpline.h
@@ -0,0 +1,15 @@
+// Copyright Epic Games, Inc. All Rights Reserved.
+
+#pragma once
+
+#include "CoreMinimal.h"
+#include "Modules/ModuleManager.h"
+
+class FVPSplineModule : public IModuleInterface
+{
+public:
+
+	/** IModuleInterface implementation */
+	virtual void StartupModule() override;
+	virtual void ShutdownModule() override;
+};
diff --git a/Engine/Plugins/Experimental/VirtualProductionUtilities/Source/VPSpline/Public/VPSplineActor.h b/Engine/Plugins/Experimental/VirtualProductionUtilities/Source/VPSpline/Public/VPSplineActor.h
new file mode 100644
--- /dev/null
+++ b/Engine/Plugins/Experimental/VirtualProductionUtilities/Source/VPSpline/Public/VPSplineActor.h
@@ -0,0 +1,98 @@
+// Copyright Epic Games, Inc. All Rights Reserved.
+#pragma once
+#include "CoreMinimal.h"
+
+#include "GameFramework/Actor.h"
+#include "CineCameraActor.h"
+
+#include "VPSplineComponent.h"
+
+#if WITH_EDITOR
+#include "ILevelSequenceEditorToolkit.h"
+#include "ISequencer.h"
+#include "LevelSequence.h"
+#include "LevelSequenceEditorBlueprintLibrary.h"
+
+#include "Channels/MovieSceneChannel.h"
+#include "Channels/MovieSceneChannelEditorData.h"
+#include "MovieSceneSection.h"
+
+#endif
+
+#include "VPSplineActor.generated.h"
+
+UCLASS(Blueprintable, Category = "VirtualProduction")
+class VPSPLINE_API AVPSplineActor : public AActor
+{
+	GENERATED_BODY()
+	
+public:
+
+	AVPSplineActor(const FObjectInitializer& ObjectInitializer);
+
+	/* VPSplineComponent attached to this actor */
+	UPROPERTY(VisibleAnywhere, BlueprintReadWrite, Category = "VPSpline")
+	TObjectPtr<UVPSplineComponent> SplineComp;
+
+	/* Component to define the attach point. Moves along the spline from CurrentPosition property*/
+	UPROPERTY(EditDefaultsOnly, Category = "VPSpline")
+	TObjectPtr<USceneComponent> SplineAttachment;
+
+	/* Normalized current parameter value to drive spline attachment*/
+	UPROPERTY(EditAnywhere, BlueprintReadWrite, Interp, Category = "VPSpline", meta = (UIMin = 0.0, UIMax = 1.0))
+	float CurrentPosition;
+
+	/* Adds or updates a keyframe position from the actor*/
+	UFUNCTION(BlueprintCallable, Category = "VPSpline")
+	void SetPointFromActor(const AActor* Actor);
+
+	/* Adds or updates a keyframe position by value*/
+	UFUNCTION(BlueprintCallable, Category = "VPSpline")
+	void SetPointByValue(const FVPSplinePointData& Data);
+
+	/* Removes a keyframe position when there is a spline point data at current position*/
+	UFUNCTION(BlueprintCallable, Category = "VPSpline")
+	void RemoveCurrentPoint();
+
+	/* Bake positions as keyframes in a current sequence*/
+	UFUNCTION(BlueprintCallable, Category = "VPSpline")
+	void BakePointsToSequence();
+
+	/* Updates CurrentPosition property if there is a next keyframe*/
+	UFUNCTION(BlueprintCallable, Category = "VPSpline")
+	void GoToNextPosition();
+
+	/* Updates CurrentPosition property if there is a previous keyframe*/
+	UFUNCTION(BlueprintCallable, Category = "VPSpline")
+	void GoToPrevPosition();
+
+	virtual void Tick(float DeltaTime) override;
+	virtual bool ShouldTickIfViewportsOnly() const override;
+	virtual class USceneComponent* GetDefaultAttachComponent() const override;
+
+#if WITH_EDITOR
+	virtual void PostEditChangeProperty(FPropertyChangedEvent& PropertyChangedEvent) override;
+	virtual void PostEditUndo() override;
+	virtual void PostEditMove(bool bFinished) override;
+#endif
+
+protected:
+	void UpdateSplineAttachment();
+
+#if WITH_EDITORONLY_DATA
+	/** Preview mesh for visualization */
+	UPROPERTY(BlueprintReadWrite, EditAnywhere, Category = "VPSpline")
+	TObjectPtr<class UStaticMesh> PreviewMesh;
+
+	UPROPERTY(Transient)
+	TObjectPtr<UStaticMeshComponent> PreviewMeshComp;
+
+	/** Determines the scale of the preview mesh */
+	UPROPERTY(Transient, EditAnywhere, Category = "VPSpline")
+	float PreviewMeshScale;
+
+	void UpdatePreviewMesh();
+
+#endif
+
+};
diff --git a/Engine/Plugins/Experimental/VirtualProductionUtilities/Source/VPSpline/Public/VPSplineComponent.h b/Engine/Plugins/Experimental/VirtualProductionUtilities/Source/VPSpline/Public/VPSplineComponent.h
new file mode 100644
--- /dev/null
+++ b/Engine/Plugins/Experimental/VirtualProductionUtilities/Source/VPSpline/Public/VPSplineComponent.h
@@ -0,0 +1,113 @@
+// Copyright Epic Games, Inc. All Rights Reserved.
+
+#pragma once
+
+#include "CoreMinimal.h"
+#include "UObject/ObjectMacros.h"
+
+#include "Components/SplineComponent.h"
+
+#include "VPSplineMetadata.h"
+#include "VPSplinePointData.h"
+
+#if WITH_EDITOR
+#include "ILevelSequenceEditorToolkit.h"
+#include "ISequencer.h"
+#include "LevelSequence.h"
+#include "LevelSequenceEditorBlueprintLibrary.h"
+
+#include "Channels/MovieSceneChannelEditorData.h"
+#include "MovieSceneSection.h"
+#endif
+
+#include "VPSplineComponent.generated.h"
+
+UCLASS(ClassGroup = Utility, ShowCategories = (Mobility), HideCategories = (Physics, Collision, Lighting, Rendering, Mobile), meta = (BlueprintSpawnableComponent))
+class VPSPLINE_API UVPSplineComponent : public USplineComponent
+{
+	GENERATED_BODY()
+
+public:
+	UVPSplineComponent(const FObjectInitializer& ObjectInitializer);
+
+	/**
+	* Defaults which are used to propagate values to spline points on instances of this in the world
+	*/
+	UPROPERTY(Category = Camera, EditDefaultsOnly)
+	FVPSplineCurveDefaults CameraSplineDefaults;
+
+	// UObject interface
+#if WITH_EDITOR
+	virtual void PostEditChangeProperty(FPropertyChangedEvent& PropertyChangedEvent) override;
+#endif
+	virtual void PostLoad() override;
+	// End of UObject interface
+
+	/** Spline component interface */
+	virtual USplineMetadata* GetSplinePointsMetadata() override;
+	virtual const USplineMetadata* GetSplinePointsMetadata() const override;
+	virtual bool AllowsSplinePointScaleEditing() const override { return false; }
+
+	void ApplyComponentInstanceData(struct FVPSplineInstanceData* ComponentInstanceData, const bool bPostUCS);
+
+	// UActorComponent Interface
+	virtual TStructOnScope<FActorComponentInstanceData> GetComponentInstanceData() const override;
+
+	/** Pointer to metadata */
+	UPROPERTY(Instanced)
+	TObjectPtr<UVPSplineMetadata> VPSplineMetadata;
+
+	/** Set focal lenght metadata at a given splint point*/
+	UFUNCTION(BlueprintCallable, Category = "VPSpline")
+	void SetFocalLengthAtSplinePoint(const int32 PointIndex, const float Value );
+
+	/** Set aperture metadata at a given spline point*/
+	UFUNCTION(BlueprintCallable, Category = "VPSpline")
+	void SetApertureAtSplinePoint(const int32 PointIndex, const float value );
+
+	/** Set focus distance metadata at a given spline point*/
+	UFUNCTION(BlueprintCallable, Category = "VPSpline")
+	void SetFocusDistanceAtSplinePoint(const int32 PointIndex, const float value);
+
+	/** Set normalized position metadata at a given spoint point*/
+	UFUNCTION(BlueprintCallable, Category = "VPSpline")
+	void SetNormalizedPositionAtSplinePoint(const int32 PointIndex, const float value);
+
+	/** Returns true if there is a spline point at given normalized position*/
+	UFUNCTION(BlueprintCallable, Category = "VPSpline")
+	bool FindSplineDataAtPosition(const float InPosition, int32& OutIndex) const;
+
+	/* Update spline point data at the given spline point*/
+	UFUNCTION(BlueprintCallable, Category = "VPSpline")
+	void UpdateSplineDataAtIndex(const int InIndex, const FVPSplinePointData& InPointData );
+
+	/* Add a new spline point data at the given normalized position*/
+	UFUNCTION(BlueprintCallable, Category = "VPSpline")
+	void AddSplineDataAtPosition(const float InPosition, const FVPSplinePointData& InPointData);
+
+	/*Calculates input key from normalized position */
+	UFUNCTION(BlueprintCallable, Category = "VPSpline")
+	float GetInputKeyAtPosition(const float InPosition);
+
+protected:
+	void SynchronizeProperties();
+
+public:
+	// Triggered when the spline is edited
+	FSimpleDelegate OnSplineEdited;
+
+};
+
+USTRUCT()
+struct FVPSplineInstanceData : public FSplineInstanceData
+{
+	GENERATED_BODY()
+public:
+	FVPSplineInstanceData() = default;
+	explicit FVPSplineInstanceData(const UVPSplineComponent* SourceComponent) : FSplineInstanceData(SourceComponent) {}
+	virtual ~FVPSplineInstanceData() = default;
+	virtual void ApplyToComponent(UActorComponent* Component, const ECacheApplyPhase CacheApplyPhase) override;
+
+	UPROPERTY()
+	TObjectPtr<UVPSplineMetadata> VPSplineMetadata;
+};
\ No newline at end of file
diff --git a/Engine/Plugins/Experimental/VirtualProductionUtilities/Source/VPSpline/Public/VPSplineMetadata.h b/Engine/Plugins/Experimental/VirtualProductionUtilities/Source/VPSpline/Public/VPSplineMetadata.h
new file mode 100644
--- /dev/null
+++ b/Engine/Plugins/Experimental/VirtualProductionUtilities/Source/VPSpline/Public/VPSplineMetadata.h
@@ -0,0 +1,66 @@
+// Copyright Epic Games, Inc. All Rights Reserved.
+
+#pragma once
+
+#include "CoreMinimal.h"
+#include "UObject/ObjectMacros.h"
+
+#include "Components/SplineComponent.h"
+
+#include "VPSplineMetadata.generated.h"
+
+USTRUCT()
+struct FVPSplineCurveDefaults
+{
+	GENERATED_BODY()
+
+		FVPSplineCurveDefaults()
+		: DefaultNormalizedPosition(-1.0f)
+		, DefaultFocalLength(35.0f)
+		, DefaultAperture(2.8f)
+		, DefaultFocusDistance(100000.f)
+	{};
+
+	UPROPERTY(EditDefaultsOnly, Category = VPSplineCurveDefaults)
+	float DefaultNormalizedPosition;
+
+	UPROPERTY(EditDefaultsOnly, Category = VPSplineCurveDefaults)
+	float DefaultFocalLength;
+
+	UPROPERTY(EditDefaultsOnly, Category = VPSplineCurveDefaults)
+	float DefaultAperture;
+
+	UPROPERTY(EditDefaultsOnly, Category = VPSplineCurveDefaults)
+	float DefaultFocusDistance;
+
+};
+
+UCLASS()
+class VPSPLINE_API UVPSplineMetadata : public USplineMetadata
+{
+	GENERATED_BODY()
+	
+public:
+	/** Insert point before index, lerping metadata between previous and next key values */
+	virtual void InsertPoint(int32 Index, float t, bool bClosedLoop) override;
+	/** Update point at index by lerping metadata between previous and next key values */
+	virtual void UpdatePoint(int32 Index, float t, bool bClosedLoop) override;
+	virtual void AddPoint(float InputKey) override;
+	virtual void RemovePoint(int32 Index) override;
+	virtual void DuplicatePoint(int32 Index) override;
+	virtual void CopyPoint(const USplineMetadata* FromSplineMetadata, int32 FromIndex, int32 ToIndex) override;
+	virtual void Reset(int32 NumPoints) override;
+	virtual void Fixup(int32 NumPoints, USplineComponent* SplineComp) override;
+
+	UPROPERTY(EditAnywhere, Category = "Point")
+	FInterpCurveFloat NormalizedPosition;
+
+	UPROPERTY(EditAnywhere, Category = "Camera")
+	FInterpCurveFloat FocalLength;
+
+	UPROPERTY(EditAnywhere, Category = "Camera")
+	FInterpCurveFloat Aperture;
+
+	UPROPERTY(EditAnywhere, Category = "Camera")
+	FInterpCurveFloat FocusDistance;
+};
diff --git a/Engine/Plugins/Experimental/VirtualProductionUtilities/Source/VPSpline/Public/VPSplinePointData.h b/Engine/Plugins/Experimental/VirtualProductionUtilities/Source/VPSpline/Public/VPSplinePointData.h
new file mode 100644
--- /dev/null
+++ b/Engine/Plugins/Experimental/VirtualProductionUtilities/Source/VPSpline/Public/VPSplinePointData.h
@@ -0,0 +1,36 @@
+// Copyright Epic Games, Inc. All Rights Reserved.
+
+#pragma once
+
+#include "CoreMinimal.h"
+#include "UObject/ObjectMacros.h"
+
+#include "VPSplinePointData.generated.h"
+
+/* Simple struct to hold spline point data*/
+USTRUCT(BlueprintType)
+struct FVPSplinePointData
+{
+	GENERATED_BODY()
+public:
+	FVPSplinePointData()
+		: FocalLength(35.0f)
+		, Aperture(2.8f)
+		, FocusDistance(100000.f)
+	{};
+
+	UPROPERTY(BlueprintReadWrite, Category="VPSpline")
+	FVector Location;
+
+	UPROPERTY(BlueprintReadWrite, Category = "VPSpline")
+	FRotator Rotation;
+
+	UPROPERTY(BlueprintReadWrite, Category = "VPSpline")
+	float FocalLength;
+
+	UPROPERTY(BlueprintReadWrite, Category = "VPSpline")
+	float Aperture;
+
+	UPROPERTY(BlueprintReadWrite, Category = "VPSpline")
+	float FocusDistance;
+};
\ No newline at end of file
diff --git a/Engine/Plugins/Experimental/VirtualProductionUtilities/Source/VPSpline/VPSpline.Build.cs b/Engine/Plugins/Experimental/VirtualProductionUtilities/Source/VPSpline/VPSpline.Build.cs
new file mode 100644
--- /dev/null
+++ b/Engine/Plugins/Experimental/VirtualProductionUtilities/Source/VPSpline/VPSpline.Build.cs
@@ -0,0 +1,75 @@
+// Copyright Epic Games, Inc. All Rights Reserved.
+
+using UnrealBuildTool;
+
+public class VPSpline : ModuleRules
+{
+	public VPSpline(ReadOnlyTargetRules Target) : base(Target)
+	{
+		PCHUsage = ModuleRules.PCHUsageMode.UseExplicitOrSharedPCHs;
+		
+		PublicIncludePaths.AddRange(
+			new string[] {
+				// ... add public include paths required here ...
+			}
+			);
+				
+		
+		PrivateIncludePaths.AddRange(
+			new string[] {
+				// ... add other private include paths required here ...
+			}
+			);
+			
+		
+		PublicDependencyModuleNames.AddRange(
+			new string[]
+			{
+				"Core",
+				// ... add other public dependencies that you statically link with here ...
+			}
+			);
+			
+		
+		PrivateDependencyModuleNames.AddRange(
+			new string[]
+			{
+				"CinematicCamera",
+				"CoreUObject",
+				"Engine",
+				"LevelSequence",
+				"MovieScene",
+				"MovieSceneTracks",
+				"Slate",
+				"SlateCore",
+				// ... add private dependencies that you statically link with here ...	
+			}
+			);
+		
+		
+		DynamicallyLoadedModuleNames.AddRange(
+			new string[]
+			{
+				// ... add any modules that your module loads dynamically here ...
+			}
+			);
+
+		if (Target.bBuildEditor == true)
+		{
+			PrivateDependencyModuleNames.AddRange(
+				new string[]
+				{
+					"LevelSequenceEditor",
+					"Sequencer",
+					"MovieSceneTools",
+					"LevelEditor",
+					"EditorFramework",
+					"UnrealEd",
+					"EditorScriptingUtilities",
+					"VPUtilitiesEditor"
+				}
+			);
+		}
+
+	}
+}
diff --git a/Engine/Plugins/Experimental/VirtualProductionUtilities/Source/VPSplineEditor/Private/VPSplineComponentVisualizer.cpp b/Engine/Plugins/Experimental/VirtualProductionUtilities/Source/VPSplineEditor/Private/VPSplineComponentVisualizer.cpp
new file mode 100644
--- /dev/null
+++ b/Engine/Plugins/Experimental/VirtualProductionUtilities/Source/VPSplineEditor/Private/VPSplineComponentVisualizer.cpp
@@ -0,0 +1,62 @@
+// Copyright Epic Games, Inc. All Rights Reserved.
+
+#include "VPSplineComponentVisualizer.h"
+#include "Internationalization/Text.h"
+#include "Math/Color.h"
+#include "CanvasTypes.h"
+#include "CanvasItem.h"
+
+#define LOCTEXT_NAMESPACE "SplineComponentVisualizer"
+
+FVPSplineComponentVisualizer::FVPSplineComponentVisualizer()
+	: FSplineComponentVisualizer()
+{
+}
+
+FVPSplineComponentVisualizer::~FVPSplineComponentVisualizer()
+{
+}
+void FVPSplineComponentVisualizer::OnRegister()
+{
+	FSplineComponentVisualizer::OnRegister();
+}
+
+void FVPSplineComponentVisualizer::DrawVisualizationHUD(const UActorComponent* Component, const FViewport* Viewport, const FSceneView* View, FCanvas* Canvas)
+{
+	FSplineComponentVisualizer::DrawVisualizationHUD(Component, Viewport, View, Canvas);
+	if (Canvas == nullptr || View == nullptr)
+	{
+		return;
+	}
+
+	const FIntRect CanvasRect = Canvas->GetViewRect();
+	float HalfX = CanvasRect.Width() / 2.f;
+	float HalfY = CanvasRect.Height() / 2.f;
+
+
+	if (const UVPSplineComponent* SplineComp = Cast<const UVPSplineComponent>(Component))
+	{
+		if (UVPSplineMetadata* Metadata = SplineComp->VPSplineMetadata)
+		{
+			int32 NumOfPoints = SplineComp->GetNumberOfSplinePoints();
+			for (int32 i = 0; i < NumOfPoints; ++i)
+			{
+				FVector Location = SplineComp->GetLocationAtSplinePoint(i, ESplineCoordinateSpace::World);
+				FPlane Plane(0, 0, 0, 0);
+				Plane = View->Project(Location);
+				const FVector Position(Plane);
+				const float DrawPositionX = FMath::FloorToFloat(HalfX + Position.X * HalfX);
+				const float DrawPositionY = FMath::FloorToFloat(HalfY + -1.f * Position.Y * HalfY);
+
+				FNumberFormattingOptions FmtOptions;
+				FmtOptions.SetMaximumFractionalDigits(3);
+				const FText Text = FText::AsNumber(Metadata->NormalizedPosition.Points[i].OutVal, &FmtOptions);
+				Canvas->DrawShadowedString(DrawPositionX, DrawPositionY, *Text.ToString(), GEngine->GetLargeFont(), FLinearColor::Yellow, FLinearColor::Black);
+			}
+		}
+		
+
+	}
+}
+
+#undef LOCTEXT_NAMESPACE
\ No newline at end of file
diff --git a/Engine/Plugins/Experimental/VirtualProductionUtilities/Source/VPSplineEditor/Private/VPSplineComponentVisualizer.h b/Engine/Plugins/Experimental/VirtualProductionUtilities/Source/VPSplineEditor/Private/VPSplineComponentVisualizer.h
new file mode 100644
--- /dev/null
+++ b/Engine/Plugins/Experimental/VirtualProductionUtilities/Source/VPSplineEditor/Private/VPSplineComponentVisualizer.h
@@ -0,0 +1,27 @@
+// Copyright Epic Games, Inc. All Rights Reserved.
+
+#pragma once
+
+#include "CoreMinimal.h"
+
+#include "HitProxies.h"
+#include "InputCoreTypes.h"
+
+#include "ComponentVisualizer.h"
+#include "SplineComponentVisualizer.h"
+
+#include "Components/SplineComponent.h"
+#include "VPSplineComponent.h"
+
+
+class FVPSplineComponentVisualizer : public FSplineComponentVisualizer
+{
+public:
+	FVPSplineComponentVisualizer();
+	virtual ~FVPSplineComponentVisualizer();
+
+	//~ Begin FComponentVisualizer Interface
+	virtual void OnRegister() override;
+	virtual void DrawVisualizationHUD(const UActorComponent* Component, const FViewport* Viewport, const FSceneView* View, FCanvas* Canvas) override;
+
+};
\ No newline at end of file
diff --git a/Engine/Plugins/Experimental/VirtualProductionUtilities/Source/VPSplineEditor/Private/VPSplineEditor.cpp b/Engine/Plugins/Experimental/VirtualProductionUtilities/Source/VPSplineEditor/Private/VPSplineEditor.cpp
new file mode 100644
--- /dev/null
+++ b/Engine/Plugins/Experimental/VirtualProductionUtilities/Source/VPSplineEditor/Private/VPSplineEditor.cpp
@@ -0,0 +1,44 @@
+// Copyright Epic Games, Inc. All Rights Reserved.
+
+#include "CoreMinimal.h"
+#include "Modules/ModuleManager.h"
+#include "VPSplineComponentVisualizer.h"
+#include "VPSplineComponent.h"
+#include "Editor.h"
+#include "Editor/UnrealEdEngine.h"
+#include "UnrealEdGlobals.h"
+
+#define LOCTEXT_NAMESPACE "VPSplineEditor"
+
+class FVPSplineEditorModule : public IModuleInterface
+{
+public:
+
+	virtual void StartupModule() override
+	{
+		VPSplineComponentName = UVPSplineComponent::StaticClass()->GetFName();
+
+		if (GUnrealEd)
+		{
+			TSharedPtr<FVPSplineComponentVisualizer> Visualizer = MakeShared<FVPSplineComponentVisualizer>();
+			GUnrealEd->RegisterComponentVisualizer(VPSplineComponentName, Visualizer);
+			Visualizer->OnRegister();
+		}
+	}
+
+	virtual void ShutdownModule() override
+	{
+		if (GUnrealEd)
+		{
+			GUnrealEd->UnregisterComponentVisualizer(VPSplineComponentName);
+		}
+	}
+
+private:
+	FName VPSplineComponentName;
+};
+
+
+#undef LOCTEXT_NAMESPACE
+
+IMPLEMENT_MODULE(FVPSplineEditorModule, VPSplineEditor)
diff --git a/Engine/Plugins/Experimental/VirtualProductionUtilities/Source/VPSplineEditor/Private/VPSplineMetadataDetails.cpp b/Engine/Plugins/Experimental/VirtualProductionUtilities/Source/VPSplineEditor/Private/VPSplineMetadataDetails.cpp
new file mode 100644
--- /dev/null
+++ b/Engine/Plugins/Experimental/VirtualProductionUtilities/Source/VPSplineEditor/Private/VPSplineMetadataDetails.cpp
@@ -0,0 +1,295 @@
+// Copyright Epic Games, Inc. All Rights Reserved.
+
+#include "VPSplineMetadataDetails.h"
+#include "VPSplineMetadata.h"
+
+#include "IDetailChildrenBuilder.h"
+#include "IDetailGroup.h"
+#include "DetailLayoutBuilder.h"
+#include "DetailWidgetRow.h"
+#include "Internationalization/Internationalization.h"
+#include "Widgets/Text/STextBlock.h"
+#include "Widgets/Input/SNumericEntryBox.h"
+#include "ScopedTransaction.h"
+#include "ComponentVisualizer.h"
+#include "Widgets/Input/SVectorInputBox.h"
+#include "Editor.h"
+
+#define LOCTEXT_NAMESPACE "FVPSplineMetadataDetails"
+
+UClass* UVPSplineMetadataDetailsFactory::GetMetadataClass() const
+{
+	return UVPSplineMetadata::StaticClass();
+}
+
+TSharedPtr<ISplineMetadataDetails> UVPSplineMetadataDetailsFactory::Create()
+{
+	return MakeShared<FVPSplineMetadataDetails>();
+}
+
+FText FVPSplineMetadataDetails::GetDisplayName() const
+{
+	return LOCTEXT("VPSplineMetadataDetails", "VPSpline");
+}
+
+template<class T>
+bool UpdateMultipleValue(TOptional<T>& CurrentValue, T InValue)
+{
+	if (!CurrentValue.IsSet())
+	{
+		CurrentValue = InValue;
+	}
+	else if (CurrentValue.IsSet() && CurrentValue.GetValue() != InValue)
+	{
+		CurrentValue.Reset();
+		return false;
+	}
+
+	return true;
+}
+
+void FVPSplineMetadataDetails::Update(USplineComponent* InSplineComponent, const TSet<int32>& InSelectedKeys)
+{
+	SplineComp = InSplineComponent;
+	SelectedKeys = InSelectedKeys;
+	NormalizedPositionValue.Reset();
+	FocalLengthValue.Reset();
+	ApertureValue.Reset();
+	FocusDistanceValue.Reset();
+	
+
+	if (IsValid(InSplineComponent))
+	{
+		bool bUpdateNormalizedPosition = true;
+		bool bUpdateFocalLength = true;
+		bool bUpdateAperture = true;
+		bool bUpdateFocusDistance = true;
+		
+
+		UVPSplineMetadata* Metadata = Cast<UVPSplineMetadata>(InSplineComponent->GetSplinePointsMetadata());
+		if (Metadata)
+		{
+			for (int32 Index : InSelectedKeys)
+			{
+				if (bUpdateNormalizedPosition)
+				{
+					bUpdateNormalizedPosition = UpdateMultipleValue(NormalizedPositionValue, Metadata->NormalizedPosition.Points[Index].OutVal);
+				}
+
+				if (bUpdateFocalLength)
+				{
+					bUpdateFocalLength = UpdateMultipleValue(FocalLengthValue, Metadata->FocalLength.Points[Index].OutVal);
+				}
+
+				if (bUpdateAperture)
+				{
+					bUpdateAperture = UpdateMultipleValue(ApertureValue, Metadata->Aperture.Points[Index].OutVal);
+				}
+
+				if (bUpdateFocusDistance)
+				{
+					bUpdateFocusDistance = UpdateMultipleValue(FocusDistanceValue, Metadata->FocusDistance.Points[Index].OutVal);
+				}
+			}
+		}
+	}
+}
+
+template<class T>
+void SetValues(FVPSplineMetadataDetails& Details, TArray<FInterpCurvePoint<T>>& Points, const T& NewValue)
+{
+	Details.SplineComp->GetSplinePointsMetadata()->Modify();
+	for (int32 Index : Details.SelectedKeys)
+	{
+		Points[Index].OutVal = NewValue;
+	}
+
+	Details.SplineComp->UpdateSpline();
+	Details.SplineComp->bSplineHasBeenEdited = true;
+	static FProperty* SplineCurvesProperty = FindFProperty<FProperty>(USplineComponent::StaticClass(), GET_MEMBER_NAME_CHECKED(USplineComponent, SplineCurves));
+	FComponentVisualizer::NotifyPropertyModified(Details.SplineComp, SplineCurvesProperty);
+	Details.Update(Details.SplineComp, Details.SelectedKeys);
+
+	if (GEditor)
+	{
+		GEditor->RedrawLevelEditingViewports(true);
+	}
+}
+
+void FVPSplineMetadataDetails::OnBeginSliderMovement()
+{
+	EditSliderValueTransaction = MakeUnique<FScopedTransaction>(LOCTEXT("EditVPSplineProperty", "Edit VPSpline Property"));
+}
+
+void FVPSplineMetadataDetails::OnEndSliderMovement(float NewValue)
+{
+	EditSliderValueTransaction.Reset();
+}
+
+UVPSplineMetadata* FVPSplineMetadataDetails::GetMetadata() const
+{
+	return SplineComp ? Cast<UVPSplineMetadata>(SplineComp->GetSplinePointsMetadata()) : nullptr;
+}
+
+void FVPSplineMetadataDetails::OnSetNormalizedPosition(float NewValue, ETextCommit::Type CommitInfo)
+{
+	if (UVPSplineMetadata* Metadata = GetMetadata())
+	{
+		const FScopedTransaction Transaction(LOCTEXT("SetNormalizedPosition", "Set spline point normalized position data"));
+		SetValues<float>(*this, Metadata->NormalizedPosition.Points, NewValue);
+	}
+}
+
+void FVPSplineMetadataDetails::OnSetFocalLength(float NewValue, ETextCommit::Type CommitInfo)
+{
+	if (UVPSplineMetadata* Metadata = GetMetadata())
+	{
+		const FScopedTransaction Transaction(LOCTEXT("SetFocalLength", "Set spline point focal length data"));
+		SetValues<float>(*this, Metadata->FocalLength.Points, NewValue);
+	}
+}
+
+void FVPSplineMetadataDetails::OnSetAperture(float NewValue, ETextCommit::Type CommitInfo)
+{
+	if (UVPSplineMetadata* Metadata = GetMetadata())
+	{
+		const FScopedTransaction Transaction(LOCTEXT("SetAperture", "Set spline point aperture data"));
+		SetValues<float>(*this, Metadata->Aperture.Points, NewValue);
+	}
+}
+
+void FVPSplineMetadataDetails::OnSetFocusDistance(float NewValue, ETextCommit::Type CommitInfo)
+{
+	if (UVPSplineMetadata* Metadata = GetMetadata())
+	{
+		const FScopedTransaction Transaction(LOCTEXT("SetFocusDistance", "Set spline point focus distance data"));
+		SetValues<float>(*this, Metadata->FocusDistance.Points, NewValue);
+	}
+}
+
+void FVPSplineMetadataDetails::GenerateChildContent(IDetailGroup& DetailGroup)
+{
+	DetailGroup.AddWidgetRow()
+		.RowTag("NormalizedPosition")
+		.Visibility(TAttribute<EVisibility>(this, &FVPSplineMetadataDetails::IsEnabled))
+		.NameContent()
+		.HAlign(HAlign_Left)
+		.VAlign(VAlign_Center)
+		[
+			SNew(STextBlock)
+			.Text(LOCTEXT("NormalizedPosition", "Normalized Position"))
+			.Font(IDetailLayoutBuilder::GetDetailFont())
+		]
+	.ValueContent()
+		.MinDesiredWidth(125.0f)
+		.MaxDesiredWidth(125.0f)
+		[
+			SNew(SNumericEntryBox<float>)
+			.Value(this, &FVPSplineMetadataDetails::GetNormalizedPosition)
+			.AllowSpin(true)
+			.MinValue(0.0f)
+			.MaxValue(1.0f)
+			.MinSliderValue(0.0f)
+			.MaxSliderValue(1.0f)
+			.OnBeginSliderMovement(this, &FVPSplineMetadataDetails::OnBeginSliderMovement)
+			.OnEndSliderMovement(this, &FVPSplineMetadataDetails::OnEndSliderMovement)
+			.UndeterminedString(LOCTEXT("Multiple", "Multiple"))
+			.OnValueCommitted(this, &FVPSplineMetadataDetails::OnSetNormalizedPosition)
+			.OnValueChanged(this, &FVPSplineMetadataDetails::OnSetNormalizedPosition, ETextCommit::Default)
+			.Font(IDetailLayoutBuilder::GetDetailFont())
+		];
+
+	DetailGroup.AddWidgetRow()
+		.RowTag("FocalLength")
+		.Visibility(TAttribute<EVisibility>(this, &FVPSplineMetadataDetails::IsEnabled))
+		.NameContent()
+		.HAlign(HAlign_Left)
+		.VAlign(VAlign_Center)
+		[
+			SNew(STextBlock)
+			.Text(LOCTEXT("FocalLength", "Focal Length"))
+			.Font(IDetailLayoutBuilder::GetDetailFont())
+		]
+	.ValueContent()
+		.MinDesiredWidth(125.0f)
+		.MaxDesiredWidth(125.0f)
+		[
+			SNew(SNumericEntryBox<float>)
+			.Value(this, &FVPSplineMetadataDetails::GetFocalLength)
+			.AllowSpin(true)
+			.OnBeginSliderMovement(this, &FVPSplineMetadataDetails::OnBeginSliderMovement)
+			.OnEndSliderMovement(this, &FVPSplineMetadataDetails::OnEndSliderMovement)
+			.MinValue(0.0f)
+			// Because we have no upper limit in MaxSliderValue, we need to "unspecify" the max value here, otherwise the spinner has a limited range,
+			//  with TNumericLimits<NumericType>::Max() as the MaxValue and the spinning increment is huge
+			.MaxValue(TOptional<float>())
+			.MinSliderValue(0.0f)
+			.MaxSliderValue(TOptional<float>()) // No upper limit
+			.UndeterminedString(LOCTEXT("Multiple", "Multiple"))
+			.OnValueCommitted(this, &FVPSplineMetadataDetails::OnSetFocalLength)
+			.OnValueChanged(this, &FVPSplineMetadataDetails::OnSetFocalLength, ETextCommit::Default)
+			.Font(IDetailLayoutBuilder::GetDetailFont())
+		];
+
+	DetailGroup.AddWidgetRow()
+		.RowTag("Aperture")
+		.Visibility(TAttribute<EVisibility>(this, &FVPSplineMetadataDetails::IsEnabled))
+		.NameContent()
+		.HAlign(HAlign_Left)
+		.VAlign(VAlign_Center)
+		[
+			SNew(STextBlock)
+			.Text(LOCTEXT("Aperture", "Aperture"))
+			.Font(IDetailLayoutBuilder::GetDetailFont())
+		]
+		.ValueContent()
+		.MinDesiredWidth(125.0f)
+		.MaxDesiredWidth(125.0f)
+		[
+			SNew(SNumericEntryBox<float>)
+			.Value(this, &FVPSplineMetadataDetails::GetAperture)
+			.AllowSpin(true)
+			.OnBeginSliderMovement(this, &FVPSplineMetadataDetails::OnBeginSliderMovement)
+			.OnEndSliderMovement(this, &FVPSplineMetadataDetails::OnEndSliderMovement)
+			.MinValue(TOptional<float>())
+			.MaxValue(TOptional<float>())
+			.MinSliderValue(TOptional<float>()) // No lower limit
+			.MaxSliderValue(TOptional<float>()) // No upper limit
+			.UndeterminedString(LOCTEXT("Multiple", "Multiple"))
+			.OnValueCommitted(this, &FVPSplineMetadataDetails::OnSetAperture)
+			.OnValueChanged(this, &FVPSplineMetadataDetails::OnSetAperture, ETextCommit::Default)
+			.Font(IDetailLayoutBuilder::GetDetailFont())
+		];
+
+	DetailGroup.AddWidgetRow()
+		.RowTag("FocusDistance")
+		.Visibility(TAttribute<EVisibility>(this, &FVPSplineMetadataDetails::IsEnabled))
+		.NameContent()
+		.HAlign(HAlign_Left)
+		.VAlign(VAlign_Center)
+		[
+			SNew(STextBlock)
+			.Text(LOCTEXT("FocusDistance", "Focus Distance"))
+			.Font(IDetailLayoutBuilder::GetDetailFont())
+		]
+		.ValueContent()
+		.MinDesiredWidth(125.0f)
+		.MaxDesiredWidth(125.0f)
+		[
+			SNew(SNumericEntryBox<float>)
+			.Value(this, &FVPSplineMetadataDetails::GetFocusDistance)
+			.AllowSpin(true)
+			.OnBeginSliderMovement(this, &FVPSplineMetadataDetails::OnBeginSliderMovement)
+			.OnEndSliderMovement(this, &FVPSplineMetadataDetails::OnEndSliderMovement)
+			.MinValue(TOptional<float>())
+			.MaxValue(TOptional<float>())
+			.MinSliderValue(TOptional<float>()) // No lower limit
+			.MaxSliderValue(TOptional<float>()) // No upper limit
+			.UndeterminedString(LOCTEXT("Multiple", "Multiple"))
+			.OnValueCommitted(this, &FVPSplineMetadataDetails::OnSetFocusDistance)
+			.OnValueChanged(this, &FVPSplineMetadataDetails::OnSetFocusDistance, ETextCommit::Default)
+			.Font(IDetailLayoutBuilder::GetDetailFont())
+		];
+}
+
+#undef LOCTEXT_NAMESPACE
\ No newline at end of file
diff --git a/Engine/Plugins/Experimental/VirtualProductionUtilities/Source/VPSplineEditor/Private/VPSplineMetadataDetails.h b/Engine/Plugins/Experimental/VirtualProductionUtilities/Source/VPSplineEditor/Private/VPSplineMetadataDetails.h
new file mode 100644
--- /dev/null
+++ b/Engine/Plugins/Experimental/VirtualProductionUtilities/Source/VPSplineEditor/Private/VPSplineMetadataDetails.h
@@ -0,0 +1,66 @@
+// Copyright Epic Games, Inc. All Rights Reserved.
+
+#pragma once
+
+#include "CoreMinimal.h"
+#include "UObject/ObjectMacros.h"
+#include "Layout/Visibility.h"
+#include "Types/SlateEnums.h"
+#include "ScopedTransaction.h"
+
+#include "SplineMetadataDetailsFactory.h"
+
+#include "VPSplineMetadataDetails.generated.h"
+
+UCLASS()
+class UVPSplineMetadataDetailsFactory : public USplineMetadataDetailsFactoryBase
+{
+	GENERATED_BODY()
+
+	public:
+	virtual ~UVPSplineMetadataDetailsFactory() {}
+	virtual TSharedPtr<ISplineMetadataDetails> Create() override;
+	virtual UClass* GetMetadataClass() const override;
+};
+
+class FVPSplineMetadataDetails : public ISplineMetadataDetails, public TSharedFromThis<FVPSplineMetadataDetails>
+{
+public:
+	virtual ~FVPSplineMetadataDetails() {}
+	virtual FName GetName() const override { return FName("VPSplineMetadataDetails"); }
+	virtual FText GetDisplayName() const override;
+	virtual void Update(USplineComponent* InSplineComponent, const TSet<int32>& InSelectedKeys) override;
+	virtual void GenerateChildContent(IDetailGroup& InGroup) override;
+
+	TOptional<float> NormalizedPositionValue;
+	TOptional<float> FocalLengthValue;
+	TOptional<float> ApertureValue;
+	TOptional<float> FocusDistanceValue;
+
+	USplineComponent* SplineComp = nullptr;
+	TSet<int32> SelectedKeys;
+
+private:
+	void OnBeginSliderMovement();
+	void OnEndSliderMovement(float NewValue);
+
+	EVisibility IsEnabled() const { return (SelectedKeys.Num() > 0) ? EVisibility::Visible : EVisibility::Collapsed; }
+	bool IsOnePointSelected() const { return SelectedKeys.Num() == 1; }
+
+
+	TOptional<float> GetNormalizedPosition() const { return NormalizedPositionValue; }
+	void OnSetNormalizedPosition(float NewValue, ETextCommit::Type CommitInfo);
+
+	TOptional<float> GetFocalLength() const { return FocalLengthValue; }
+	void OnSetFocalLength(float NewValue, ETextCommit::Type CommitInfo);
+
+	TOptional<float> GetAperture() const { return ApertureValue; }
+	void OnSetAperture(float NewValue, ETextCommit::Type CommitInfo);
+
+	TOptional<float> GetFocusDistance() const { return FocusDistanceValue; }
+	void OnSetFocusDistance(float NewValue, ETextCommit::Type CommitInfo);
+
+	class UVPSplineMetadata* GetMetadata() const;
+private:
+	TUniquePtr<FScopedTransaction> EditSliderValueTransaction;
+};
diff --git a/Engine/Plugins/Experimental/VirtualProductionUtilities/Source/VPSplineEditor/VPSplineEditor.Build.cs b/Engine/Plugins/Experimental/VirtualProductionUtilities/Source/VPSplineEditor/VPSplineEditor.Build.cs
new file mode 100644
--- /dev/null
+++ b/Engine/Plugins/Experimental/VirtualProductionUtilities/Source/VPSplineEditor/VPSplineEditor.Build.cs
@@ -0,0 +1,40 @@
+// Copyright Epic Games, Inc. All Rights Reserved.
+
+using UnrealBuildTool;
+
+public class VPSplineEditor : ModuleRules
+{
+	public VPSplineEditor(ReadOnlyTargetRules Target) : base(Target)
+	{
+		PCHUsage = ModuleRules.PCHUsageMode.UseExplicitOrSharedPCHs;
+		
+        PrivateIncludePaths.Add("VPSpline/Private");
+
+		PrivateDependencyModuleNames.AddRange(
+			new string[] {
+                "Core",
+				"CoreUObject",
+                "ComponentVisualizers",
+				"DetailCustomizations",
+				"Engine",
+				"LevelSequence",
+				"LevelSequenceEditor",
+				"Sequencer",
+				"SlateCore",
+				"Slate",
+				"UnrealEd",
+				"VPSpline",
+				"PropertyEditor",
+				"InputCore",
+			});
+
+		PublicDependencyModuleNames.AddRange(
+			new string[] {
+			}
+		);
+
+		PrivateIncludePathModuleNames.AddRange(
+			new string[] {
+			});
+	}
+}
diff --git a/Engine/Plugins/Experimental/VirtualProductionUtilities/Source/VPUtilitiesEditor/Public/VPUtilitiesEditorSettings.h b/Engine/Plugins/Experimental/VirtualProductionUtilities/Source/VPUtilitiesEditor/Public/VPUtilitiesEditorSettings.h
--- a/Engine/Plugins/Experimental/VirtualProductionUtilities/Source/VPUtilitiesEditor/Public/VPUtilitiesEditorSettings.h
+++ b/Engine/Plugins/Experimental/VirtualProductionUtilities/Source/VPUtilitiesEditor/Public/VPUtilitiesEditorSettings.h
@@ -80,4 +80,7 @@ public:
 	/** GestureManager class to use by the ScoutingSubsystem */
 	UPROPERTY(config)
 	TArray<FSoftClassPath> AdditionnalClassToLoad;
+
+	UPROPERTY(config, meta = (AllowedClasses = "/Script/Engine.StaticMesh"))
+	FSoftObjectPath VPSplinePreviewMeshPath;
 };
\ No newline at end of file
diff --git a/Engine/Plugins/Experimental/VirtualProductionUtilities/VirtualProductionUtilities.uplugin b/Engine/Plugins/Experimental/VirtualProductionUtilities/VirtualProductionUtilities.uplugin
--- a/Engine/Plugins/Experimental/VirtualProductionUtilities/VirtualProductionUtilities.uplugin
+++ b/Engine/Plugins/Experimental/VirtualProductionUtilities/VirtualProductionUtilities.uplugin
@@ -48,6 +48,16 @@
 			"Name": "VPRolesEditor",
 			"Type": "Editor",
 			"LoadingPhase": "Default"
+		},
+		{
+			"Name": "VPSpline",
+			"Type": "Runtime",
+			"LoadingPhase": "Default"
+		},
+		{
+			"Name": "VPSplineEditor",
+			"Type": "Editor",
+			"LoadingPhase": "PostEngineInit"
 		}
 	],
 	"Plugins": [
