VPBookmark
- removed hard-coded constructor refs to assets
- added asset paths to BaseVirtualProductionUtilities.ini

diff --git a/Engine/Plugins/Experimental/VirtualProductionUtilities/Config/BaseVirtualProductionUtilities.ini b/Engine/Plugins/Experimental/VirtualProductionUtilities/Config/BaseVirtualProductionUtilities.ini
--- a/Engine/Plugins/Experimental/VirtualProductionUtilities/Config/BaseVirtualProductionUtilities.ini
+++ b/Engine/Plugins/Experimental/VirtualProductionUtilities/Config/BaseVirtualProductionUtilities.ini
@@ -3,10 +3,18 @@
 +ClassRedirects=(OldName="VRTool",NewName="/Script/VPUtilitiesEditor.VRTool")
 +FunctionRedirects=(OldName="VPViewportTickableActorBase.LockLocation",NewName="VPViewportTickableActorBase.EditorLockLocation")
 
-
 [/Script/VPUtilitiesEditor.VPUtilitiesEditorSettings]
 ScoutingSubsystemEditorUtilityClassPath=/VirtualProductionUtilities/BP_VPScoutingSubsystemHelpers.BP_VPScoutingSubsystemHelpers_C
 GestureManagerEditorUtilityClassPath=/VirtualProductionUtilities/VRTools/BP_GestureManager.BP_GestureManager_C
 +AdditionnalClassToLoad=/VirtualProductionUtilities/VR/VirtualScoutingInteractor.VirtualScoutingInteractor
 +AdditionnalClassToLoad=/VirtualProductionUtilities/VR/VirtualScoutingTeleporter.VirtualScoutingTeleporter
-VPSplinePreviewMeshPath=/VirtualProductionUtilities/Editor/Meshes/SM_VPSplineArraw.SM_VPSplineArraw
\ No newline at end of file
+VPSplinePreviewMeshPath=/VirtualProductionUtilities/Editor/Meshes/SM_VPSplineArraw.SM_VPSplineArraw
+
+
+[/Script/VPUtilities.VPBookmarkSettings]
+
+BookmarkMeshPath=/VirtualProductionUtilities/Bookmark/BookmarkVPMesh.BookmarkVPMesh
+BookmarkMaterialPath=/VirtualProductionUtilities/Bookmark/MI_BookmarkVPNoUser.MI_BookmarkVPNoUser
+BookmarkSplineMeshPath=/Engine/VREditor/FloatingText/LineSegmentCylinder.LineSegmentCylinder
+BookmarkSplineMeshMaterialPath=/Engine/VREditor/LaserPointer/LaserPointerMaterialInst.LaserPointerMaterialInst
+BookmarkLabelMaterialPath=/VirtualProductionUtilities/Bookmark/CameraFacingTextMaterial.CameraFacingTextMaterial
\ No newline at end of file
diff --git a/Engine/Plugins/Experimental/VirtualProductionUtilities/Source/VPUtilities/Private/VPBookmarkActor.cpp b/Engine/Plugins/Experimental/VirtualProductionUtilities/Source/VPUtilities/Private/VPBookmarkActor.cpp
--- a/Engine/Plugins/Experimental/VirtualProductionUtilities/Source/VPUtilities/Private/VPBookmarkActor.cpp
+++ b/Engine/Plugins/Experimental/VirtualProductionUtilities/Source/VPUtilities/Private/VPBookmarkActor.cpp
@@ -13,6 +13,8 @@
 #include "VPBookmark.h"
 #include "Materials/MaterialInstanceDynamic.h"
 #include "Materials/Material.h"
+#include "VPSettings.h"
+#include "UObject/UObjectGlobals.h"
 
 
 AVPBookmarkActor::AVPBookmarkActor(const FObjectInitializer& ObjectInitializer)
@@ -26,19 +28,17 @@ AVPBookmarkActor::AVPBookmarkActor(const FObjectInitializer& ObjectInitializer)
 	BookmarkMesh = CreateDefaultSubobject<UStaticMeshComponent>(TEXT("BookmarkMesh"));
 	SetRootComponent(BookmarkMesh);
 
-	static ConstructorHelpers::FObjectFinder<UStaticMesh> BookmarkStaticMesh(TEXT("/VirtualProductionUtilities/Bookmark/BookmarkVPMesh.BookmarkVPMesh"));
-	if (BookmarkStaticMesh.Succeeded())
-	{
-		FoundMesh = BookmarkStaticMesh.Object;
-	}
+	FSoftObjectPath BookmarkMeshAssetPath = GetDefault<UVPBookmarkSettings>()->BookmarkMeshPath;
+
+	UStaticMesh* BookmarkStaticMesh = LoadObject<UStaticMesh>(nullptr, *BookmarkMeshAssetPath.ToString(), nullptr, LOAD_None, nullptr);
+	
+	BookmarkMesh->SetStaticMesh(BookmarkStaticMesh);
 
-	BookmarkMesh->SetStaticMesh(FoundMesh);
+	FSoftObjectPath BookmarkMaterialAssetPath = GetDefault<UVPBookmarkSettings>()->BookmarkMaterialPath;
 
-	static ConstructorHelpers::FObjectFinder<UMaterialInterface> BookmarkMaterialInstance(TEXT("/VirtualProductionUtilities/Bookmark/MI_BookmarkVPNoUser.MI_BookmarkVPNoUser"));
-	if (BookmarkStaticMesh.Succeeded())
-	{
-		BookmarkMaterial = BookmarkMaterialInstance.Object;
-	}
+	UMaterialInterface* BookmarkMaterialInstance = LoadObject<UMaterialInterface>(nullptr, *BookmarkMaterialAssetPath.ToString(), nullptr, LOAD_None, nullptr);
+
+	BookmarkMaterial = BookmarkMaterialInstance;
 		
 	BookmarkColor = FLinearColor(0.817708f, 0.107659f, 0.230336f);
 
@@ -48,21 +48,18 @@ AVPBookmarkActor::AVPBookmarkActor(const FObjectInitializer& ObjectInitializer)
 	SplineMesh->SetupAttachment(BookmarkMesh);
 	SplineMesh->SetVisibility(false);
 
-	static ConstructorHelpers::FObjectFinder<UStaticMesh> SplineStaticMesh(TEXT("/Engine/VREditor/FloatingText/LineSegmentCylinder.LineSegmentCylinder"));
-	if (SplineStaticMesh.Succeeded())
-	{
-		FoundMesh = SplineStaticMesh.Object;
-	};
-	SplineMesh->SetStaticMesh(FoundMesh);
+	FSoftObjectPath SplineMeshAssetPath = GetDefault<UVPBookmarkSettings>()->BookmarkSplineMeshPath;
 
-	static ConstructorHelpers::FObjectFinder<UMaterialInterface> LaserMaterial(TEXT("/Engine/VREditor/LaserPointer/LaserPointerMaterialInst.LaserPointerMaterialInst"));
-	if (LaserMaterial.Succeeded())
-	{
-		SplineMaterial = LaserMaterial.Object;
-	}
-	SplineMesh->SetMaterial(0, SplineMaterial);
-	
+	UStaticMesh* SplineStaticMesh = LoadObject<UStaticMesh>(nullptr, *SplineMeshAssetPath.ToString(), nullptr, LOAD_None, nullptr);
+
+	SplineMesh->SetStaticMesh(SplineStaticMesh);
 
+	FSoftObjectPath SplineMaterialPath = GetDefault<UVPBookmarkSettings>()->BookmarkSplineMeshMaterialPath;
+
+	UMaterialInterface* SplineMaterialInstance = LoadObject<UMaterialInterface>(nullptr, *SplineMaterialPath.ToString(), nullptr, LOAD_None, nullptr);
+
+	SplineMesh->SetMaterial(0, SplineMaterialInstance);
+	
 	//Text Render
 	NameTextRender = CreateDefaultSubobject<UTextRenderComponent>(TEXT("NameTextRender"));
 	NameTextRender->SetupAttachment(BookmarkMesh);
@@ -70,12 +67,11 @@ AVPBookmarkActor::AVPBookmarkActor(const FObjectInitializer& ObjectInitializer)
 	NameTextRender->AddRelativeLocation(FVector(0.f, 0.f, 70.f));
 	NameTextRender->HorizontalAlignment = EHorizTextAligment(EHTA_Center);
 
-	static ConstructorHelpers::FObjectFinder<UMaterial> FoundMaterial(TEXT("/VirtualProductionUtilities/Bookmark/CameraFacingTextMaterial.CameraFacingTextMaterial"));
-	if (FoundMaterial.Succeeded())
-	{
-		TextMaterial = FoundMaterial.Object;
-	}
-	NameTextRender->SetMaterial(0, TextMaterial);
+	FSoftObjectPath LabelMaterialPath = GetDefault<UVPBookmarkSettings>()->BookmarkLabelMaterialPath;
+
+	UMaterialInterface* LabelMaterialInstance = LoadObject<UMaterialInterface>(nullptr, *LabelMaterialPath.ToString(), nullptr, LOAD_None, nullptr);
+
+	NameTextRender->SetMaterial(0, LabelMaterialInstance);
 
 }
 
@@ -163,26 +159,31 @@ void AVPBookmarkActor::OnActorDroppedFromTransform_Implementation()
 
 void AVPBookmarkActor::UpdateBookmarkColor(FLinearColor Color)
 {
-
-	UMaterialInterface* Material = BookmarkMesh->GetMaterial(0);
-
-	if (Material && !Material->IsA<UMaterialInstanceDynamic>())
+	if (BookmarkMesh->GetStaticMesh() != nullptr)
 	{
-		UMaterialInstanceDynamic* BookmarkMaterialInstance = UMaterialInstanceDynamic::Create(Material, BookmarkMaterial, TEXT("BookmarkMaterial"));
+		UMaterialInterface* Material = BookmarkMesh->GetMaterial(0);
 
-		BookmarkMaterialInstance->ClearParameterValues();
-		BookmarkMaterialInstance->SetVectorParameterValue(TEXT("UserColor"), Color);
-
-		for (int32 i = 0; i < BookmarkMesh->GetNumMaterials(); i++)
+		if (BookmarkMaterial != nullptr)
 		{
-			BookmarkMesh->SetMaterial(i, BookmarkMaterialInstance);
+			if (Material && !Material->IsA<UMaterialInstanceDynamic>())
+			{
+				UMaterialInstanceDynamic* BookmarkMaterialInstance = UMaterialInstanceDynamic::Create(Material, BookmarkMaterial, TEXT("BookmarkMaterial"));
+
+				BookmarkMaterialInstance->ClearParameterValues();
+				BookmarkMaterialInstance->SetVectorParameterValue(TEXT("UserColor"), Color);
+
+				for (int32 i = 0; i < BookmarkMesh->GetNumMaterials(); i++)
+				{
+					BookmarkMesh->SetMaterial(i, BookmarkMaterialInstance);
+				}
+			}
+			else //If DMIs are already setup set the color value
+			{
+				DynamicMaterial = Cast<UMaterialInstanceDynamic>(Material);
+				DynamicMaterial->SetVectorParameterValue(TEXT("UserColor"), Color);
+			}
 		}
 	}
-	else //If DMIs are already setup set the color value
-	{
-		DynamicMaterial = Cast<UMaterialInstanceDynamic>(Material);
-		DynamicMaterial->SetVectorParameterValue(TEXT("UserColor"), Color);
-	}
 }
 #if WITH_EDITOR
 void AVPBookmarkActor::PostEditChangeProperty(FPropertyChangedEvent& PropertyChangedEvent)
diff --git a/Engine/Plugins/Experimental/VirtualProductionUtilities/Source/VPUtilities/Public/VPSettings.h b/Engine/Plugins/Experimental/VirtualProductionUtilities/Source/VPUtilities/Public/VPSettings.h
--- a/Engine/Plugins/Experimental/VirtualProductionUtilities/Source/VPUtilities/Public/VPSettings.h
+++ b/Engine/Plugins/Experimental/VirtualProductionUtilities/Source/VPUtilities/Public/VPSettings.h
@@ -39,6 +39,7 @@ public:
 	static UVPSettings* GetVPSettings();
 
 #if WITH_EDITORONLY_DATA
+
 	/** When enabled, the virtual production role(s) will be displayed in the main editor UI. */
 	UPROPERTY(config, EditAnywhere, Category="Virtual Production")
 	bool bShowRoleInEditor;
@@ -71,4 +72,31 @@ private:
 	FOnGetRoles RolesGetter;
 
 	friend class UVirtualProductionRolesSubsystem;
+};
+
+UCLASS(Config = VirtualProductionUtilities)
+class VPUTILITIES_API UVPBookmarkSettings : public UObject
+{
+	GENERATED_BODY()
+		//~ Bookmark assets //
+public:
+	/** Bookmark mesh to use for Bookmark actor */
+	UPROPERTY(config, EditAnywhere, Category = "VP Bookmark", meta = (AllowedClasses = "/Script/Engine.StaticMesh", DisplayName = "Bookmark Mesh"))
+	FSoftObjectPath BookmarkMeshPath;
+
+	/** Bookmark mesh material to use for Bookmark actor */
+	UPROPERTY(config, EditAnywhere, Category = "VP Bookmark", meta = (AllowedClasses = "/Script/Engine.MaterialInterface", DisplayName = "Bookmark Mesh Material"))
+	FSoftObjectPath BookmarkMaterialPath;
+
+	/** Bookmark spline mesh to use for Bookmark actor */
+	UPROPERTY(config, EditAnywhere, Category = "VP Bookmark", meta = (AllowedClasses = "/Script/Engine.StaticMesh", DisplayName = "Bookmark Spline Mesh"))
+	FSoftObjectPath BookmarkSplineMeshPath;
+
+	/** Bookmark spline mesh material to use for Bookmark actor */
+	UPROPERTY(config, EditAnywhere, Category = "VP Bookmark", meta = (AllowedClasses = "/Script/Engine.MaterialInterface", DisplayName = "Bookmark Spline Material"))
+	FSoftObjectPath BookmarkSplineMeshMaterialPath;
+
+	/** Bookmark label material class to use for Bookmark actor */
+	UPROPERTY(config, EditAnywhere, Category = "VP Bookmark", meta = (AllowedClasses = "/Script/Engine.MaterialInterface", DisplayName = "Bookmark Label Material"))
+	FSoftObjectPath BookmarkLabelMaterialPath;
 };
\ No newline at end of file
