Expose SetActorLabel in VCamBlueprintFunctionLibrary

diff --git a/Engine/Plugins/Experimental/VirtualProductionUtilities/Source/VPUtilities/Private/VPBookmarkActor.cpp b/Engine/Plugins/Experimental/VirtualProductionUtilities/Source/VPUtilities/Private/VPBookmarkActor.cpp
--- a/Engine/Plugins/Experimental/VirtualProductionUtilities/Source/VPUtilities/Private/VPBookmarkActor.cpp
+++ b/Engine/Plugins/Experimental/VirtualProductionUtilities/Source/VPUtilities/Private/VPBookmarkActor.cpp
@@ -73,6 +73,9 @@ AVPBookmarkActor::AVPBookmarkActor(const FObjectInitializer& ObjectInitializer)
 
 	NameTextRender->SetMaterial(0, LabelMaterialInstance);
 
+	CameraComponent = CreateDefaultSubobject<UCineCameraComponent>(TEXT("Camera"));
+	CameraComponent->SetupAttachment(BookmarkMesh);
+
 }
 
 void AVPBookmarkActor::Tick(float DeltaSeconds)
diff --git a/Engine/Plugins/Experimental/VirtualProductionUtilities/Source/VPUtilities/Public/VPBookmarkActor.h b/Engine/Plugins/Experimental/VirtualProductionUtilities/Source/VPUtilities/Public/VPBookmarkActor.h
--- a/Engine/Plugins/Experimental/VirtualProductionUtilities/Source/VPUtilities/Public/VPBookmarkActor.h
+++ b/Engine/Plugins/Experimental/VirtualProductionUtilities/Source/VPUtilities/Public/VPBookmarkActor.h
@@ -8,6 +8,7 @@
 #include "VPBookmark.h"
 #include "IVPBookmarkProvider.h"
 #include "IVPInteraction.h"
+#include "CineCameraComponent.h"
 #include "VPBookmarkActor.generated.h"
 
 class UStaticMeshComponent;
@@ -38,6 +39,10 @@ public:
 	UPROPERTY(BlueprintReadOnly, EditDefaultsOnly, Category = "Components")
 	TObjectPtr<USplineMeshComponent> SplineMesh;
 
+	/**CineCamera */
+	UPROPERTY(BlueprintReadOnly, EditDefaultsOnly, Category = "Components")
+	TObjectPtr<UCineCameraComponent> CameraComponent;
+
 	/** Color of Bookmark in MU Session */
 	UPROPERTY(BlueprintReadWrite, EditAnywhere, Category = "Bookmark", meta = (MultiLine = "true", ExposeOnSpawn = "true"))
 	FLinearColor BookmarkColor;
