VPBookmark - bugfix - component constructors used now to create object copies, not object refs.
- force camera component on bookmark actor to not be visible in editor or game view

diff --git a/Engine/Plugins/Experimental/VirtualProductionUtilities/Source/VPUtilities/Private/VPBookmarkActor.cpp b/Engine/Plugins/Experimental/VirtualProductionUtilities/Source/VPUtilities/Private/VPBookmarkActor.cpp
--- a/Engine/Plugins/Experimental/VirtualProductionUtilities/Source/VPUtilities/Private/VPBookmarkActor.cpp
+++ b/Engine/Plugins/Experimental/VirtualProductionUtilities/Source/VPUtilities/Private/VPBookmarkActor.cpp
@@ -12,82 +12,133 @@
 #include "VPBlueprintLibrary.h"
 #include "VPBookmark.h"
 #include "Materials/MaterialInstanceDynamic.h"
-#include "Materials/Material.h"
 #include "VPSettings.h"
-#include "UObject/UObjectGlobals.h"
-
 
 AVPBookmarkActor::AVPBookmarkActor(const FObjectInitializer& ObjectInitializer)
-	:Super(ObjectInitializer)
-	{
+	: Super(ObjectInitializer)
+{
 	PrimaryActorTick.bCanEverTick = true;
 	PrimaryActorTick.bStartWithTickEnabled = true;
 
 	//Root Component
 
-	BookmarkMesh = CreateDefaultSubobject<UStaticMeshComponent>(TEXT("BookmarkMesh"));
-	SetRootComponent(BookmarkMesh);
+	BookmarkMeshComponent = CreateDefaultSubobject<UStaticMeshComponent>(TEXT("BookmarkMesh"));
+	SetRootComponent(BookmarkMeshComponent);
+	
+	if (!IsRunningCommandlet())
+	{
+		static ConstructorHelpers::FObjectFinder<UStaticMesh> MeshFinder(*GetDefault<UVPBookmarkSettings>()->BookmarkMeshPath.ToString());
+		if (MeshFinder.Succeeded())
+		{
+			BookmarkStaticMesh = MeshFinder.Object;
+		}
+		else
+		{
+			UE_LOG(LogVPUtilities, Warning, TEXT("Failed to load bookmark spline: %s"), *GetDefault<UVPBookmarkSettings>()->BookmarkMeshPath.ToString());
+		}
 
-	FSoftObjectPath BookmarkMeshAssetPath = GetDefault<UVPBookmarkSettings>()->BookmarkMeshPath;
+	}
 
-	UStaticMesh* BookmarkStaticMesh = LoadObject<UStaticMesh>(nullptr, *BookmarkMeshAssetPath.ToString(), nullptr, LOAD_None, nullptr);
+	BookmarkMeshComponent->SetStaticMesh(BookmarkStaticMesh);
 	
-	BookmarkMesh->SetStaticMesh(BookmarkStaticMesh);
-
-	FSoftObjectPath BookmarkMaterialAssetPath = GetDefault<UVPBookmarkSettings>()->BookmarkMaterialPath;
+	if (!IsRunningCommandlet())
+	{
+		static ConstructorHelpers::FObjectFinder<UMaterialInterface> MaterialFinder(*GetDefault<UVPBookmarkSettings>()->BookmarkMaterialPath.ToString());
 
-	UMaterialInterface* BookmarkMaterialInstance = LoadObject<UMaterialInterface>(nullptr, *BookmarkMaterialAssetPath.ToString(), nullptr, LOAD_None, nullptr);
+		if (MaterialFinder.Succeeded())
+		{
+			BookmarkMaterial = MaterialFinder.Object;
+		}
+		else
+		{
+			UE_LOG(LogVPUtilities, Warning, TEXT("Failed to load bookmark material: %s"), *GetDefault<UVPBookmarkSettings>()->BookmarkMaterialPath.ToString());
+		}
+	}
+	//Apply found material to all material slots on the BookmarkMeshComponent
+	for (int32 i = 0; i < BookmarkMeshComponent->GetNumMaterials(); i++)
+	{ 
+		BookmarkMeshComponent->SetMaterial(i, BookmarkMaterial);
+	}
 
-	BookmarkMaterial = BookmarkMaterialInstance;
-		
 	BookmarkColor = FLinearColor(0.817708f, 0.107659f, 0.230336f);
 
-	//SplineMesh
-	SplineMesh = CreateDefaultSubobject<USplineMeshComponent>(TEXT("SplineMesh"));
-	SplineMesh->SetMobility(EComponentMobility::Movable);
-	SplineMesh->SetupAttachment(BookmarkMesh);
-	SplineMesh->SetVisibility(false);
-
-	FSoftObjectPath SplineMeshAssetPath = GetDefault<UVPBookmarkSettings>()->BookmarkSplineMeshPath;
-
-	UStaticMesh* SplineStaticMesh = LoadObject<UStaticMesh>(nullptr, *SplineMeshAssetPath.ToString(), nullptr, LOAD_None, nullptr);
+	//SplineMesh setup
+	SplineMeshComponent = CreateDefaultSubobject<USplineMeshComponent>(TEXT("SplineMesh"));
+	SplineMeshComponent->SetMobility(EComponentMobility::Movable);
+	SplineMeshComponent->SetupAttachment(BookmarkMeshComponent);
+	SplineMeshComponent->SetVisibility(false);
+	
+	if (!IsRunningCommandlet())
+	{
+		static ConstructorHelpers::FObjectFinder<UStaticMesh> SplineMeshFinder(*GetDefault<UVPBookmarkSettings>()->BookmarkSplineMeshPath.ToString());
 
-	SplineMesh->SetStaticMesh(SplineStaticMesh);
+		if (SplineMeshFinder.Succeeded())
+		{
+			SplineStaticMesh = SplineMeshFinder.Object;
+		}
+		else
+		{
+			UE_LOG(LogVPUtilities, Warning, TEXT("Failed to load bookmark material: %s"), *GetDefault<UVPBookmarkSettings>()->BookmarkSplineMeshPath.ToString());
+		}
+	}
 
-	FSoftObjectPath SplineMaterialPath = GetDefault<UVPBookmarkSettings>()->BookmarkSplineMeshMaterialPath;
+	SplineMeshComponent->SetStaticMesh(SplineStaticMesh);
 
-	UMaterialInterface* SplineMaterialInstance = LoadObject<UMaterialInterface>(nullptr, *SplineMaterialPath.ToString(), nullptr, LOAD_None, nullptr);
+	if (!IsRunningCommandlet())
+	{
+		static ConstructorHelpers::FObjectFinder<UMaterialInterface> SplineMaterialFinder(*GetDefault<UVPBookmarkSettings>()->BookmarkSplineMeshMaterialPath.ToString());
 
-	SplineMesh->SetMaterial(0, SplineMaterialInstance);
+		if (SplineMaterialFinder.Succeeded())
+		{
+			SplineMaterialInstance = SplineMaterialFinder.Object;
+		}
+		else
+		{
+			UE_LOG(LogVPUtilities, Warning, TEXT("Failed to load bookmark spline material: %s"), *GetDefault<UVPBookmarkSettings>()->BookmarkSplineMeshMaterialPath.ToString());
+		}
+	}
 	
-	//Text Render
-	NameTextRender = CreateDefaultSubobject<UTextRenderComponent>(TEXT("NameTextRender"));
-	NameTextRender->SetupAttachment(BookmarkMesh);
-	NameTextRender->SetWorldSize(36);
-	NameTextRender->AddRelativeLocation(FVector(0.f, 0.f, 70.f));
-	NameTextRender->HorizontalAlignment = EHorizTextAligment(EHTA_Center);
-
-	FSoftObjectPath LabelMaterialPath = GetDefault<UVPBookmarkSettings>()->BookmarkLabelMaterialPath;
+	SplineMeshComponent->SetMaterial(0, SplineMaterialInstance);
+
+	//Text Render setup
+	NameTextRenderComponent = CreateDefaultSubobject<UTextRenderComponent>(TEXT("NameTextRender"));
+	NameTextRenderComponent->SetupAttachment(BookmarkMeshComponent);
+	NameTextRenderComponent->SetWorldSize(36);
+	NameTextRenderComponent->AddRelativeLocation(FVector(0.f, 0.f, 70.f));
+	NameTextRenderComponent->HorizontalAlignment = EHorizTextAligment(EHTA_Center);
+	
+	if (!IsRunningCommandlet())
+	{
+		static ConstructorHelpers::FObjectFinder<UMaterialInterface> LabelMaterialFinder(*GetDefault<UVPBookmarkSettings>()->BookmarkLabelMaterialPath.ToString());
 
-	UMaterialInterface* LabelMaterialInstance = LoadObject<UMaterialInterface>(nullptr, *LabelMaterialPath.ToString(), nullptr, LOAD_None, nullptr);
+		if (LabelMaterialFinder.Succeeded())
+		{
+			LabelMaterialInstance = LabelMaterialFinder.Object;
+		}
+		else
+		{
+			UE_LOG(LogVPUtilities, Warning, TEXT("Failed to load bookmark spline material: %s"), *GetDefault<UVPBookmarkSettings>()->BookmarkLabelMaterialPath.ToString());
+		}
+	}
 
-	NameTextRender->SetMaterial(0, LabelMaterialInstance);
+	NameTextRenderComponent->SetMaterial(0, LabelMaterialInstance);
 
 	CameraComponent = CreateDefaultSubobject<UCineCameraComponent>(TEXT("Camera"));
-	CameraComponent->SetupAttachment(BookmarkMesh);
-
+	CameraComponent->SetupAttachment(BookmarkMeshComponent);
+	CameraComponent->SetVisibility(false, true);
+	CameraComponent->SetVisibleFlag(false);
+#if WITH_EDITOR
+	CameraComponent->SetCameraMesh(nullptr);// set CameraMesh to null when we are in editor view
+#endif
 }
 
 void AVPBookmarkActor::Tick(float DeltaSeconds)
-	{
+{
 	Super::Tick(DeltaSeconds);
 
-
-	NameTextRender->SetWorldRotation(FRotator(0.f, 0.f, 0.f));
-
+	NameTextRenderComponent->SetWorldRotation(FRotator(0.f, 0.f, 0.f));
 	BookmarkRotation = GetActorRotation();
-	BookmarkMesh->SetWorldRotation(FRotator(0.f,BookmarkRotation.Yaw,0.f));
-
+	BookmarkMeshComponent->SetWorldRotation(FRotator(0.f,BookmarkRotation.Yaw,0.f));
 
 #if WITH_EDITOR
 	if (GIsEditor)
@@ -105,7 +156,6 @@ void AVPBookmarkActor::OnBookmarkActivation_Implementation(UVPBookmark* Bookmark
 	UE_LOG(LogVPUtilities, Display, TEXT("Bookmark Created"));
 }
 
-
 void AVPBookmarkActor::OnBookmarkChanged_Implementation(UVPBookmark* BookmarkOut)
 {
 	AActor* BookmarkActor = BookmarkOut->GetAssociatedBookmarkActor();
@@ -118,30 +168,25 @@ void AVPBookmarkActor::OnBookmarkChanged_Implementation(UVPBookmark* BookmarkOut
 	UE_LOG(LogVPUtilities, Display, TEXT("Bookmark Updated"));
 }
 
-
 void AVPBookmarkActor::UpdateBookmarkSplineMeshIndicator_Implementation()
 {
-	UVPBlueprintLibrary::VPBookmarkSplineMeshIndicatorSetStartAndEnd(SplineMesh);
+	UVPBlueprintLibrary::VPBookmarkSplineMeshIndicatorSetStartAndEnd(SplineMeshComponent);
 }
 
-
 void AVPBookmarkActor::HideBookmarkSplineMeshIndicator_Implementation()
 {
-	UVPBlueprintLibrary::VPBookmarkSplineMeshIndicatorDisable(SplineMesh);
+	UVPBlueprintLibrary::VPBookmarkSplineMeshIndicatorDisable(SplineMeshComponent);
 }
 
-
 void AVPBookmarkActor::GenerateBookmarkName_Implementation()
 {
 	FString GeneratedNumber;
 	FString GeneratedLetter;
 	UVPBookmarkBlueprintLibrary::CreateVPBookmarkName(this, FString("Bookmark %n"), GeneratedNumber, GeneratedLetter);
 		
-	NameTextRender->SetText(FText::AsCultureInvariant(GeneratedNumber));
-
+	NameTextRenderComponent->SetText(FText::AsCultureInvariant(GeneratedNumber));
 }
 
-
 //VP Interaction Interface Events
 
 void AVPBookmarkActor::OnActorDroppedFromCarry_Implementation()
@@ -162,22 +207,22 @@ void AVPBookmarkActor::OnActorDroppedFromTransform_Implementation()
 
 void AVPBookmarkActor::UpdateBookmarkColor(FLinearColor Color)
 {
-	if (BookmarkMesh->GetStaticMesh() != nullptr)
+	if (BookmarkMeshComponent->GetStaticMesh() != nullptr)
 	{
-		UMaterialInterface* Material = BookmarkMesh->GetMaterial(0);
+		UMaterialInterface* Material = BookmarkMeshComponent->GetMaterial(0);
 
 		if (BookmarkMaterial != nullptr)
 		{
 			if (Material && !Material->IsA<UMaterialInstanceDynamic>())
-			{
-				UMaterialInstanceDynamic* BookmarkMaterialInstance = UMaterialInstanceDynamic::Create(Material, BookmarkMaterial, TEXT("BookmarkMaterial"));
+			{			
+				DynamicMaterial = UMaterialInstanceDynamic::Create(Material, DynamicMaterial, TEXT("BookmarkMaterial"));
 
-				BookmarkMaterialInstance->ClearParameterValues();
-				BookmarkMaterialInstance->SetVectorParameterValue(TEXT("UserColor"), Color);
+				DynamicMaterial->ClearParameterValues();
+				DynamicMaterial->SetVectorParameterValue(TEXT("UserColor"), Color);
 
-				for (int32 i = 0; i < BookmarkMesh->GetNumMaterials(); i++)
+				for (int32 i = 0; i < BookmarkMeshComponent->GetNumMaterials(); i++)
 				{
-					BookmarkMesh->SetMaterial(i, BookmarkMaterialInstance);
+					BookmarkMeshComponent->SetMaterial(i, DynamicMaterial);
 				}
 			}
 			else //If DMIs are already setup set the color value
diff --git a/Engine/Plugins/Experimental/VirtualProductionUtilities/Source/VPUtilities/Public/VPBookmarkActor.h b/Engine/Plugins/Experimental/VirtualProductionUtilities/Source/VPUtilities/Public/VPBookmarkActor.h
--- a/Engine/Plugins/Experimental/VirtualProductionUtilities/Source/VPUtilities/Public/VPBookmarkActor.h
+++ b/Engine/Plugins/Experimental/VirtualProductionUtilities/Source/VPUtilities/Public/VPBookmarkActor.h
@@ -29,15 +29,15 @@ public:
 
 	/** Mesh Representation in the world */
 	UPROPERTY(BlueprintReadOnly, EditDefaultsOnly, Category = "Components")
-	TObjectPtr<UStaticMeshComponent> BookmarkMesh;
+	TObjectPtr<UStaticMeshComponent> BookmarkMeshComponent;
 
 	/** Textrender to display bookmark name */
 	UPROPERTY(BlueprintReadOnly, EditDefaultsOnly, Category = "Components")
-	TObjectPtr<UTextRenderComponent> NameTextRender;
+	TObjectPtr<UTextRenderComponent> NameTextRenderComponent;
 
 	/** Splinemesh */
 	UPROPERTY(BlueprintReadOnly, EditDefaultsOnly, Category = "Components")
-	TObjectPtr<USplineMeshComponent> SplineMesh;
+	TObjectPtr<USplineMeshComponent> SplineMeshComponent;
 
 	/**CineCamera */
 	UPROPERTY(BlueprintReadOnly, EditDefaultsOnly, Category = "Components")
@@ -59,7 +59,7 @@ public:
 	UPROPERTY(BlueprintReadWrite, Category = "Snapshot")
 	TObjectPtr<UTexture2D> SnapshotTexture;
 
-	/**Update the mesh color and BookmarkColor variable. Intended for use with multiuser initialisation*/
+	/**Update the mesh color and BookmarkColor variable. Intended for use with multiuser initialization*/
 	UFUNCTION(BlueprintCallable, CallInEditor, Category = "Default")
 	void UpdateBookmarkColor(FLinearColor Color);
 
@@ -111,11 +111,23 @@ public:
 	
 private:
 
-	UMaterial* TextMaterial;
+	FRotator BookmarkRotation; //Actor's rotation
 
-	UStaticMesh* FoundMesh;
+	UPROPERTY()
+	TObjectPtr<UStaticMesh> BookmarkStaticMesh; //Mesh to use main static mesh component
 
-	UMaterialInterface* SplineMaterial;
+	UPROPERTY()
+	TObjectPtr<UMaterial> TextMaterial; 
 
-	FRotator BookmarkRotation;
+	UPROPERTY()
+	TObjectPtr<UStaticMesh> SplineStaticMesh;
+
+	UPROPERTY()
+	TObjectPtr<UMaterialInterface> SplineMaterial;
+
+	UPROPERTY()
+	TObjectPtr<UMaterialInterface> SplineMaterialInstance;
+
+	UPROPERTY()
+	TObjectPtr<UMaterialInterface> LabelMaterialInstance;
 };
