Add missing icon for VP Roles toolbar menu

diff --git a/Engine/Plugins/Experimental/VirtualProductionUtilities/Source/VPRolesEditor/Private/VPRolesEditorStyle.cpp b/Engine/Plugins/Experimental/VirtualProductionUtilities/Source/VPRolesEditor/Private/VPRolesEditorStyle.cpp
--- a/Engine/Plugins/Experimental/VirtualProductionUtilities/Source/VPRolesEditor/Private/VPRolesEditorStyle.cpp
+++ b/Engine/Plugins/Experimental/VirtualProductionUtilities/Source/VPRolesEditor/Private/VPRolesEditorStyle.cpp
@@ -1,6 +1,8 @@
 // Copyright Epic Games, Inc. All Rights Reserved.
 
 #include "VPRolesEditorStyle.h"
+
+#include "Brushes/SlateImageBrush.h"
 #include "Interfaces/IPluginManager.h"
 #include "Framework/Application/SlateApplication.h"
 #include "Slate/SlateGameResources.h"
@@ -31,7 +33,7 @@ FName FVPRolesEditorStyle::GetStyleSetName()
 	return StyleSetName;
 }
 
-#define IMAGE_BRUSH( RelativePath, ... ) FSlateImageBrush( Style->RootToContentDir( RelativePath, TEXT(".png") ), __VA_ARGS__ )
+#define IMAGE_BRUSH_SVG( RelativePath, ... ) FSlateVectorImageBrush( Style->RootToContentDir( RelativePath, TEXT(".svg") ), __VA_ARGS__ )
 
 const FVector2D Icon16x16(16.0f, 16.0f);
 const FVector2D Icon40x40(40.0f, 40.0f);
@@ -44,10 +46,7 @@ TSharedRef<FSlateStyleSet> FVPRolesEditorStyle::Create()
 	FSlateBrush AddRoleBrush = *FAppStyle::Get().GetBrush("Icons.Plus");
 	AddRoleBrush.TintColor = FStyleColors::AccentGreen;
 
-	Style->Set("VPRolesEditor.TabIcon", new IMAGE_BRUSH(TEXT("VPRolesButtonIcon_40x"), Icon16x16));
-	Style->Set("VPRolesEditor.OpenMenu", new IMAGE_BRUSH(TEXT("VPRolesButtonIcon_40x"), Icon40x40));
 	Style->Set("VPRolesEditor.AddRole", new FSlateBrush(AddRoleBrush));
-
 	return Style;
 }
 
@@ -57,8 +56,4 @@ const ISlateStyle& FVPRolesEditorStyle::Get()
 	return *StyleInstance;
 }
 
-#undef IMAGE_BRUSH
-#undef BOX_BRUSH
-#undef BORDER_BRUSH
-#undef TTF_FONT
-#undef OTF_FONT
+#undef IMAGE_BRUSH_SVG
diff --git a/Engine/Source/Runtime/SlateCore/Private/Styling/StarshipCoreStyle.cpp b/Engine/Source/Runtime/SlateCore/Private/Styling/StarshipCoreStyle.cpp
--- a/Engine/Source/Runtime/SlateCore/Private/Styling/StarshipCoreStyle.cpp
+++ b/Engine/Source/Runtime/SlateCore/Private/Styling/StarshipCoreStyle.cpp
@@ -415,6 +415,8 @@ TSharedRef<ISlateStyle> FStarshipCoreStyle::Create()
 		// Constraint Manager Icons
 		Style->Set("Icons.ConstraintManager.LookAt", new IMAGE_BRUSH_SVG("Starship/Common/LookAt", Icon16x16));
 		Style->Set("Icons.ConstraintManager.ParentHierarchy", new IMAGE_BRUSH_SVG("Starship/Common/ParentHierarchy", Icon16x16));
+
+		Style->Set("Icons.Role", new IMAGE_BRUSH_SVG("Starship/Common/Role", Icon16x16));
 	}
 
 	// Tool panels
