Add missing file for icon change

diff --git a/Engine/Plugins/Experimental/VirtualProductionUtilities/Source/VPRolesEditor/Private/VPRolesEditorModule.cpp b/Engine/Plugins/Experimental/VirtualProductionUtilities/Source/VPRolesEditor/Private/VPRolesEditorModule.cpp
--- a/Engine/Plugins/Experimental/VirtualProductionUtilities/Source/VPRolesEditor/Private/VPRolesEditorModule.cpp
+++ b/Engine/Plugins/Experimental/VirtualProductionUtilities/Source/VPRolesEditor/Private/VPRolesEditorModule.cpp
@@ -74,14 +74,14 @@ void FVPRolesEditorModule::ExtendLevelEditorToolbar()
 		}
 		
 		UToolMenu* Menu = UToolMenus::Get()->ExtendMenu("LevelEditor.LevelEditorToolBar.User");
-		
+
 		FToolMenuEntry VPRolesEntry = FToolMenuEntry::InitComboButton(
 			"VPRolesMenu",
 			FUIAction(),
 			FOnGetContent::CreateRaw(this, &FVPRolesEditorModule::GenerateVPRolesLevelEditorToolbarMenu),
 			LOCTEXT("LevelEditorToolbarVPRolesButtonLabel", "VP Roles"),
 			LOCTEXT("LevelEditorToolbarVPRolesButtonTooltip", "Edit VP Roles"),
-			FSlateIcon(FVPRolesEditorStyle::Get().GetStyleSetName(), "VPRolesEditor.TabIcon"));
+			FSlateIcon(FAppStyle::GetAppStyleSetName(), "Icons.Role"));
 		
 		VPRolesEntry.StyleNameOverride = "CalloutToolbar";
 		Menu->FindOrAddSection("VPRoles").AddEntry(VPRolesEntry);
