AsyncLoading: Fix "ExistingObject == Object " assert in AsyncLoading2.cpp when booting QAGCookedEditor by avoiding to reload some virtual production cooked blueprint packages.
Fix the non-iostore path so that QAGCookedEditor works with loose cooked and pak files too, at least without optional data.

Tweak [/Script/VPUtilitiesEditor.VPUtilitiesEditorSettings] config to specify these correct class names to async load on editor startup:
+AdditionnalClassToLoad=/VirtualProductionUtilities/VR/VirtualScoutingInteractor.VirtualScoutingInteractor_C
+AdditionnalClassToLoad=/VirtualProductionUtilities/VR/VirtualScoutingTeleporter.VirtualScoutingTeleporter_C

#rnx

diff --git a/Engine/Plugins/Experimental/VirtualProductionUtilities/Config/BaseVirtualProductionUtilities.ini b/Engine/Plugins/Experimental/VirtualProductionUtilities/Config/BaseVirtualProductionUtilities.ini
--- a/Engine/Plugins/Experimental/VirtualProductionUtilities/Config/BaseVirtualProductionUtilities.ini
+++ b/Engine/Plugins/Experimental/VirtualProductionUtilities/Config/BaseVirtualProductionUtilities.ini
@@ -6,8 +6,8 @@
 [/Script/VPUtilitiesEditor.VPUtilitiesEditorSettings]
 ScoutingSubsystemEditorUtilityClassPath=/VirtualProductionUtilities/BP_VPScoutingSubsystemHelpers.BP_VPScoutingSubsystemHelpers_C
 GestureManagerEditorUtilityClassPath=/VirtualProductionUtilities/VRTools/BP_GestureManager.BP_GestureManager_C
-+AdditionnalClassToLoad=/VirtualProductionUtilities/VR/VirtualScoutingInteractor.VirtualScoutingInteractor
-+AdditionnalClassToLoad=/VirtualProductionUtilities/VR/VirtualScoutingTeleporter.VirtualScoutingTeleporter
++AdditionnalClassToLoad=/VirtualProductionUtilities/VR/VirtualScoutingInteractor.VirtualScoutingInteractor_C
++AdditionnalClassToLoad=/VirtualProductionUtilities/VR/VirtualScoutingTeleporter.VirtualScoutingTeleporter_C
 VPSplinePreviewMeshPath=/VirtualProductionUtilities/Editor/Meshes/SM_VPSplineArraw.SM_VPSplineArraw
 
 
@@ -17,4 +17,4 @@ BookmarkMeshPath=/VirtualProductionUtilities/Bookmark/BookmarkVPMesh.BookmarkVPM
 BookmarkMaterialPath=/VirtualProductionUtilities/Bookmark/MI_BookmarkVPNoUser.MI_BookmarkVPNoUser
 BookmarkSplineMeshPath=/Engine/VREditor/FloatingText/LineSegmentCylinder.LineSegmentCylinder
 BookmarkSplineMeshMaterialPath=/Engine/VREditor/LaserPointer/LaserPointerMaterialInst.LaserPointerMaterialInst
-BookmarkLabelMaterialPath=/VirtualProductionUtilities/Bookmark/CameraFacingTextMaterial.CameraFacingTextMaterial
\ No newline at end of file
+BookmarkLabelMaterialPath=/VirtualProductionUtilities/Bookmark/CameraFacingTextMaterial.CameraFacingTextMaterial
diff --git a/Engine/Source/Runtime/CoreUObject/Private/UObject/LinkerLoad.cpp b/Engine/Source/Runtime/CoreUObject/Private/UObject/LinkerLoad.cpp
--- a/Engine/Source/Runtime/CoreUObject/Private/UObject/LinkerLoad.cpp
+++ b/Engine/Source/Runtime/CoreUObject/Private/UObject/LinkerLoad.cpp
@@ -3312,9 +3312,9 @@ bool FLinkerLoad::VerifyImportInner(const int32 ImportIndex, FString& WarningSuf
 				nullptr /* InstancingContext */, nullptr /* DiffPackagePath */);
 		}
 #if WITH_IOSTORE_IN_EDITOR
-		if (Package && Package->HasAnyPackageFlags(PKG_Cooked))
+		if (Package && Package->HasAnyPackageFlags(PKG_Cooked) && Package->GetPackageId().IsValid())
 		{
-			// Cooked packages loaded with EDL is always fully loaded and has no attached
+			// Cooked packages loaded from iostore are always fully loaded and has no attached
 			// linkers. Static find all imported objects from this package.
 			check(Package->IsFullyLoaded());
 			Import.XObject = Package;
@@ -3371,7 +3371,7 @@ bool FLinkerLoad::VerifyImportInner(const int32 ImportIndex, FString& WarningSuf
 		// or when circular dependency happen, get the linker so we are able to create the import properly at a later time.
 		// When loading editor data never consider the package fully loaded and resolve the linker anyway, for cooked data, assign the linker if one is associated witht the package
 #if WITH_IOSTORE_IN_EDITOR
-		if (Package && !Package->HasAnyPackageFlags(PKG_Cooked))
+		if (Package && !(Package->HasAnyPackageFlags(PKG_Cooked) && Package->GetPackageId().IsValid()))
 #endif
 		{
 			const bool bWasFullyLoaded = Package && Package->IsFullyLoaded() && FPlatformProperties::RequiresCookedData();
