Fixed unity compilation, very likely a side effect of include cleanup.

diff --git a/Engine/Plugins/Experimental/VirtualProductionUtilities/Source/VPUtilities/Private/VPFullScreenUserWidget.cpp b/Engine/Plugins/Experimental/VirtualProductionUtilities/Source/VPUtilities/Private/VPFullScreenUserWidget.cpp
--- a/Engine/Plugins/Experimental/VirtualProductionUtilities/Source/VPUtilities/Private/VPFullScreenUserWidget.cpp
+++ b/Engine/Plugins/Experimental/VirtualProductionUtilities/Source/VPUtilities/Private/VPFullScreenUserWidget.cpp
@@ -4,6 +4,7 @@
 
 #include "Components/PostProcessComponent.h"
 #include "Engine/Engine.h"
+#include "Engine/GameViewportClient.h"
 #include "Engine/TextureRenderTarget2D.h"
 #include "Engine/UserInterfaceSettings.h"
 #include "GameFramework/WorldSettings.h"
