Attempt to fix NoPCH game build error

diff --git a/Engine/Plugins/Experimental/VirtualProductionUtilities/Source/VPSpline/Private/VPSplineComponent.cpp b/Engine/Plugins/Experimental/VirtualProductionUtilities/Source/VPSpline/Private/VPSplineComponent.cpp
--- a/Engine/Plugins/Experimental/VirtualProductionUtilities/Source/VPSpline/Private/VPSplineComponent.cpp
+++ b/Engine/Plugins/Experimental/VirtualProductionUtilities/Source/VPSpline/Private/VPSplineComponent.cpp
@@ -2,6 +2,8 @@
 
 #include "VPSplineComponent.h"
 
+#include "Engine/Engine.h"
+
 #include "VPSpline.h"
 #include "VPSplineLog.h"
 
