Dynamic Material Instance - Fix Copy to Original does not immedieately update

diff --git a/Engine/Plugins/Experimental/VirtualProductionUtilities/Source/VPMaterialsEditor/Private/SMaterialDynamicWidgets.cpp b/Engine/Plugins/Experimental/VirtualProductionUtilities/Source/VPMaterialsEditor/Private/SMaterialDynamicWidgets.cpp
--- a/Engine/Plugins/Experimental/VirtualProductionUtilities/Source/VPMaterialsEditor/Private/SMaterialDynamicWidgets.cpp
+++ b/Engine/Plugins/Experimental/VirtualProductionUtilities/Source/VPMaterialsEditor/Private/SMaterialDynamicWidgets.cpp
@@ -2,11 +2,14 @@
 
 #include "SMaterialDynamicWidgets.h"
 
-#include "PropertyCustomizationHelpers.h"
-#include "ScopedTransaction.h"
-#include "SMaterialDynamicParametersPanelWidget.h"
 #include "Components/ActorComponent.h"
+#include "Editor.h"
+#include "MaterialEditingLibrary.h"
 #include "Materials/MaterialInstanceDynamic.h"
+#include "PropertyCustomizationHelpers.h"
+#include "SMaterialDynamicParametersPanelWidget.h"
+#include "ScopedTransaction.h"
+#include "Subsystems/AssetEditorSubsystem.h"
 #include "Widgets/Input/SButton.h"
 #include "Widgets/Layout/SBox.h"
 #include "Widgets/Text/STextBlock.h"
@@ -200,6 +203,14 @@ FReply SMaterialDynamicView::OnCopyToOriginalButtonClicked() const
 	ParentMaterialInstance->UpdateStaticPermutation();
 	ParentMaterialInstance->Modify();
 
+#if WITH_EDITOR
+	// Check if the Parent is being Edited actively in a Window somewhere
+	if (GEditor->GetEditorSubsystem<UAssetEditorSubsystem>()->FindEditorForAsset(ParentMaterialInstance, false))
+	{
+		UMaterialEditingLibrary::RecompileMaterial(ParentMaterialInstance->GetMaterial());
+	}
+#endif
+
 	return FReply::Handled();
 }
 
diff --git a/Engine/Plugins/Experimental/VirtualProductionUtilities/Source/VPMaterialsEditor/VPMaterialsEditor.Build.cs b/Engine/Plugins/Experimental/VirtualProductionUtilities/Source/VPMaterialsEditor/VPMaterialsEditor.Build.cs
--- a/Engine/Plugins/Experimental/VirtualProductionUtilities/Source/VPMaterialsEditor/VPMaterialsEditor.Build.cs
+++ b/Engine/Plugins/Experimental/VirtualProductionUtilities/Source/VPMaterialsEditor/VPMaterialsEditor.Build.cs
@@ -17,6 +17,7 @@ public class VPMaterialsEditor : ModuleRules
 				"InputCore",
 				"Landscape",
 				"PropertyEditor",
+				"MaterialEditor",
 				"Slate",
 				"SlateCore",
 				"UnrealEd",
