Fix uninitialized structs warning

diff --git a/Engine/Plugins/Experimental/VirtualProductionUtilities/Source/VPSpline/Public/VPSplinePointData.h b/Engine/Plugins/Experimental/VirtualProductionUtilities/Source/VPSpline/Public/VPSplinePointData.h
--- a/Engine/Plugins/Experimental/VirtualProductionUtilities/Source/VPSpline/Public/VPSplinePointData.h
+++ b/Engine/Plugins/Experimental/VirtualProductionUtilities/Source/VPSpline/Public/VPSplinePointData.h
@@ -14,7 +14,9 @@ struct FVPSplinePointData
 	GENERATED_BODY()
 public:
 	FVPSplinePointData()
-		: FocalLength(35.0f)
+		: Location(FVector::ZeroVector)
+		, Rotation(FRotator::ZeroRotator)
+		, FocalLength(35.0f)
 		, Aperture(2.8f)
 		, FocusDistance(100000.f)
 	{};
diff --git a/Engine/Plugins/VirtualProduction/RemoteControl/Source/WebRemoteControl/Private/RemoteControlRequest.h b/Engine/Plugins/VirtualProduction/RemoteControl/Source/WebRemoteControl/Private/RemoteControlRequest.h
--- a/Engine/Plugins/VirtualProduction/RemoteControl/Source/WebRemoteControl/Private/RemoteControlRequest.h
+++ b/Engine/Plugins/VirtualProduction/RemoteControl/Source/WebRemoteControl/Private/RemoteControlRequest.h
@@ -305,7 +305,7 @@ struct FRCPresetExposePropertyRequest : public FRCRequest
 	 * Whether to automatically enable the edit condition for the exposed property.
 	 */
 	UPROPERTY()
-	bool EnableEditCondition;
+	bool EnableEditCondition = true;
 };
 
 /**
