Pixel streaming output providers should use CineCamera post processing settings instead of attaching global component to world settings

diff --git a/Engine/Plugins/Experimental/VirtualCamera/Source/VCamCore/Private/Output/VCamOutputProviderBase.cpp b/Engine/Plugins/Experimental/VirtualCamera/Source/VCamCore/Private/Output/VCamOutputProviderBase.cpp
--- a/Engine/Plugins/Experimental/VirtualCamera/Source/VCamCore/Private/Output/VCamOutputProviderBase.cpp
+++ b/Engine/Plugins/Experimental/VirtualCamera/Source/VCamCore/Private/Output/VCamOutputProviderBase.cpp
@@ -154,6 +154,10 @@ void UVCamOutputProviderBase::CreateUMG()
 	UMGWidget->SetDisplayTypes(DisplayType, DisplayType, DisplayType);
 	UMGWidget->PostProcessDisplayType.bReceiveHardwareInput = true;
 
+	// The post process material created by EVPWidgetDisplayType::PostProcess should affect only our VCam, not the entire world.
+	// This allows you to have multiple virtual camera in the world affect different viewports.
+	UMGWidget->SetCustomPostProcessSettingsSource(TargetCamera.Get());
+
 #if WITH_EDITOR
 	UMGWidget->SetAllTargetViewports(GetTargetLevelViewport());
 #endif
@@ -239,6 +243,8 @@ void UVCamOutputProviderBase::NotifyWidgetOfComponentChange() const
 {
 	if (UMGWidget && UMGWidget->IsDisplayed())
 	{
+		UMGWidget->SetCustomPostProcessSettingsSource(TargetCamera.Get());
+		
 		UUserWidget* DisplayedWidget = UMGWidget->GetWidget();
 		if (IsValid(DisplayedWidget))
 		{
diff --git a/Engine/Plugins/Experimental/VirtualProductionUtilities/Source/VPUtilities/Private/VPFullScreenUserWidget.cpp b/Engine/Plugins/Experimental/VirtualProductionUtilities/Source/VPUtilities/Private/VPFullScreenUserWidget.cpp
--- a/Engine/Plugins/Experimental/VirtualProductionUtilities/Source/VPUtilities/Private/VPFullScreenUserWidget.cpp
+++ b/Engine/Plugins/Experimental/VirtualProductionUtilities/Source/VPUtilities/Private/VPFullScreenUserWidget.cpp
@@ -314,7 +314,19 @@ FVPFullScreenUserWidget_PostProcess::FVPFullScreenUserWidget_PostProcess()
 	, PostProcessMaterialInstance(nullptr)
 	, WidgetRenderer(nullptr)
 	, CurrentWidgetDrawSize(FIntPoint::ZeroValue)
+{}
+
+void FVPFullScreenUserWidget_PostProcess::SetCustomPostProcessSettingsSource(TWeakObjectPtr<UObject> InCustomPostProcessSettingsSource)
 {
+	CustomPostProcessSettingsSource = InCustomPostProcessSettingsSource;
+	
+	const bool bIsRunning = PostProcessMaterialInstance != nullptr;
+	if (bIsRunning)
+	{
+		// Save us from creating another struct member: PostProcessMaterialInstance is always created with UWorld as outer.
+		UWorld* World = CastChecked<UWorld>(PostProcessMaterialInstance->GetOuter());
+		InitPostProcessComponent(World);
+	}
 }
 
 bool FVPFullScreenUserWidget_PostProcess::Display(UWorld* World, UUserWidget* Widget, bool bInRenderToTextureOnly, float InDPIScale)
@@ -324,7 +336,7 @@ bool FVPFullScreenUserWidget_PostProcess::Display(UWorld* World, UUserWidget* Wi
 	bool bOk = CreateRenderer(World, Widget, InDPIScale);
 	if (!bRenderToTextureOnly)
 	{
-		bOk &= CreatePostProcessComponent(World);
+		bOk &= InitPostProcessComponent(World);
 	}
 
 	return bOk;
@@ -350,17 +362,21 @@ TSharedPtr<SVirtualWindow> FVPFullScreenUserWidget_PostProcess::GetSlateWindow()
 	return SlateWindow;
 }
 
-bool FVPFullScreenUserWidget_PostProcess::CreatePostProcessComponent(UWorld* World)
+bool FVPFullScreenUserWidget_PostProcess::InitPostProcessComponent(UWorld* World)
 {
 	ReleasePostProcessComponent();
 	if (World && PostProcessMaterial)
 	{
-		AWorldSettings* WorldSetting = World->GetWorldSettings();
-		PostProcessComponent = NewObject<UPostProcessComponent>(WorldSetting, NAME_None, RF_Transient);
-		PostProcessComponent->bEnabled = true;
-		PostProcessComponent->bUnbound = true;
-		PostProcessComponent->RegisterComponent();
-
+		const bool bUseExternalPostProcess = CustomPostProcessSettingsSource.IsValid();
+		if (!bUseExternalPostProcess)
+		{
+			AWorldSettings* WorldSetting = World->GetWorldSettings();
+			PostProcessComponent = NewObject<UPostProcessComponent>(WorldSetting, NAME_None, RF_Transient);
+			PostProcessComponent->bEnabled = true;
+			PostProcessComponent->bUnbound = true;
+			PostProcessComponent->RegisterComponent();
+		}
+		
 		PostProcessMaterialInstance = UMaterialInstanceDynamic::Create(PostProcessMaterial, World);
 
 		// set the parameter immediately
@@ -368,9 +384,12 @@ bool FVPFullScreenUserWidget_PostProcess::CreatePostProcessComponent(UWorld* Wor
 		PostProcessMaterialInstance->SetVectorParameterValue(NAME_TintColorAndOpacity, PostProcessTintColorAndOpacity);
 		PostProcessMaterialInstance->SetScalarParameterValue(NAME_OpacityFromTexture, PostProcessOpacityFromTexture);
 
-		PostProcessComponent->Settings.WeightedBlendables.Array.SetNumZeroed(1);
-		PostProcessComponent->Settings.WeightedBlendables.Array[0].Weight = 1.f;
-		PostProcessComponent->Settings.WeightedBlendables.Array[0].Object = PostProcessMaterialInstance;
+		if (FPostProcessSettings* const PostProcessSettings = GetPostProcessSettings())
+		{
+			PostProcessSettings->WeightedBlendables.Array.SetNumZeroed(1);
+			PostProcessSettings->WeightedBlendables.Array[0].Weight = 1.f;
+			PostProcessSettings->WeightedBlendables.Array[0].Object = PostProcessMaterialInstance;
+		}
 	}
 
 	return PostProcessComponent && PostProcessMaterialInstance;
@@ -378,10 +397,28 @@ bool FVPFullScreenUserWidget_PostProcess::CreatePostProcessComponent(UWorld* Wor
 
 void FVPFullScreenUserWidget_PostProcess::ReleasePostProcessComponent()
 {
-	if (PostProcessComponent)
+	const bool bIsRunning = PostProcessMaterialInstance != nullptr;
+	if (!bIsRunning)
+	{
+		return;
+	}
+	
+	const bool bNeedsToResetExternalSettings = CustomPostProcessSettingsSource.IsValid();
+	if (FPostProcessSettings* Settings = GetPostProcessSettings()
+		; bNeedsToResetExternalSettings && Settings)
+	{
+		const int32 Index = Settings->WeightedBlendables.Array.IndexOfByPredicate([this](const FWeightedBlendable& Blendable){ return Blendable.Object == PostProcessMaterialInstance; });
+		if (Index != INDEX_NONE)
+		{
+			Settings->WeightedBlendables.Array.RemoveAt(Index);
+		}
+	}
+	// CustomPostProcessSettingsSource may have gone stale
+	else if (PostProcessComponent)
 	{
 		PostProcessComponent->UnregisterComponent();
 	}
+	
 	PostProcessComponent = nullptr;
 	PostProcessMaterialInstance = nullptr;
 }
@@ -627,6 +664,32 @@ void FVPFullScreenUserWidget_PostProcess::UnRegisterHitTesterWithViewport()
 	CustomHitTestPath.Reset();
 }
 
+FPostProcessSettings* FVPFullScreenUserWidget_PostProcess::GetPostProcessSettings() const
+{
+	if (PostProcessComponent)
+	{
+		return &PostProcessComponent->Settings;
+	}
+
+	if (!CustomPostProcessSettingsSource.IsValid())
+	{
+		UE_LOG(LogVPUtilities, Warning, TEXT("CustomPostProcessSettingsSource has become stale"))
+		return nullptr;
+	}
+
+	// The easiest way without overcomplicating the API with an additional callback would be to look for the first struct property.
+	// We could always extend our API to accept a callback that extracts the FPostProcessSettings from the UObject instead.
+	for (TFieldIterator<FStructProperty> StructIt(CustomPostProcessSettingsSource.Get()->GetClass()); StructIt; ++StructIt)
+	{
+		if (StructIt->Struct == FPostProcessSettings::StaticStruct())
+		{
+			return StructIt->ContainerPtrToValuePtr<FPostProcessSettings>(CustomPostProcessSettingsSource.Get());
+		}
+	}
+
+	return nullptr;
+}
+
 /////////////////////////////////////////////////////
 // UVPFullScreenUserWidget
 
@@ -812,6 +875,11 @@ void UVPFullScreenUserWidget::SetDisplayTypes(EVPWidgetDisplayType InEditorDispl
 	PIEDisplayType = InPIEDisplayType;
 }
 
+void UVPFullScreenUserWidget::SetCustomPostProcessSettingsSource(TWeakObjectPtr<UObject> InCustomPostProcessSettingsSource)
+{
+	PostProcessDisplayType.SetCustomPostProcessSettingsSource(InCustomPostProcessSettingsSource);
+}
+
 void UVPFullScreenUserWidget::InitWidget()
 {
 	// Don't do any work if Slate is not initialized
diff --git a/Engine/Plugins/Experimental/VirtualProductionUtilities/Source/VPUtilities/Public/VPFullScreenUserWidget.h b/Engine/Plugins/Experimental/VirtualProductionUtilities/Source/VPUtilities/Public/VPFullScreenUserWidget.h
--- a/Engine/Plugins/Experimental/VirtualProductionUtilities/Source/VPUtilities/Public/VPFullScreenUserWidget.h
+++ b/Engine/Plugins/Experimental/VirtualProductionUtilities/Source/VPUtilities/Public/VPFullScreenUserWidget.h
@@ -6,6 +6,7 @@
 #include "Blueprint/UserWidget.h"
 #include "Components/WidgetComponent.h"
 #include "CompositingElement.h"
+#include "Templates/NonNullPointer.h"
 #include "VPFullScreenUserWidget.generated.h"
 
 class FWidgetRenderer;
@@ -25,7 +26,6 @@ class UWorld;
 class SLevelViewport;
 #endif
 
-
 UENUM(BlueprintType)
 enum class EVPWidgetDisplayType : uint8
 {
@@ -45,7 +45,6 @@ struct FVPFullScreenUserWidget_Viewport
 {
 	GENERATED_BODY()
 
-public:
 	FVPFullScreenUserWidget_Viewport();
 	bool Display(UWorld* World, UUserWidget* Widget, float InDPIScale);
 	void Hide(UWorld* World);
@@ -73,8 +72,10 @@ struct FVPFullScreenUserWidget_PostProcess
 {
 	GENERATED_BODY()
 
-public:
 	FVPFullScreenUserWidget_PostProcess();
+
+	void SetCustomPostProcessSettingsSource(TWeakObjectPtr<UObject> InCustomPostProcessSettingsSource);
+	
 	bool Display(UWorld* World, UUserWidget* Widget, bool bInRenderToTextureOnly, float InDPIScale);
 	void Hide(UWorld* World);
 	void Tick(UWorld* World, float DeltaSeconds);
@@ -82,7 +83,7 @@ public:
 	TSharedPtr<SVirtualWindow> VPUTILITIES_API GetSlateWindow() const;
 
 private:
-	bool CreatePostProcessComponent(UWorld* World);
+	bool InitPostProcessComponent(UWorld* World);
 	void ReleasePostProcessComponent();
 
 	bool CreateRenderer(UWorld* World, UUserWidget* Widget, float InDPIScale);
@@ -95,6 +96,8 @@ private:
 	void RegisterHitTesterWithViewport(UWorld* World);
 	void UnRegisterHitTesterWithViewport();
 
+	FPostProcessSettings* GetPostProcessSettings() const;
+
 public:
 	/**
 	 * Post process material used to display the widget.
@@ -153,6 +156,7 @@ public:
 	/** If set, use this viewport instead of GetFirstActiveLevelViewport() */
 	TWeakPtr<SLevelViewport> TargetViewport;
 #endif
+	
 private:
 	/** Post process component used to add the material to the post process chain. */
 	UPROPERTY(Transient)
@@ -179,6 +183,12 @@ private:
 
 	/** Only render to the UTextureRenderTarget2D - do not output to the final viewport. */
 	bool bRenderToTextureOnly;
+	
+	/**
+	 * Optional. Some object that contains a FPostProcessSettings property. These settings will be used for PostProcessMaterialInstance.
+	 * E.g. VCam uses this to use post process from a specific cine camera.
+	 */
+	TWeakObjectPtr<UObject> CustomPostProcessSettingsSource;
 };
 
 /**
@@ -208,6 +218,14 @@ public:
 	virtual void Tick(float DeltaTime);
 
 	void SetDisplayTypes(EVPWidgetDisplayType InEditorDisplayType, EVPWidgetDisplayType InGameDisplayType, EVPWidgetDisplayType InPIEDisplayType);
+	
+	/**
+	 * If using EVPWidgetDisplayType::PostProcess, you can specify a custom post process settings that should be modified.
+	 * By default, a new post process component is added to AWorldSettings.
+	 *
+	 * @param InCustomPostProcessSettingsSource An object containing a FPostProcessSettings UPROPERTY()
+	 */
+	void SetCustomPostProcessSettingsSource(TWeakObjectPtr<UObject> InCustomPostProcessSettingsSource);
 
 protected:
 	void InitWidget();
