Outliner: Fix an issue causing two "Virtual Production" categories in the filter bar

diff --git a/Engine/Plugins/Experimental/VirtualProductionUtilities/Source/VPUtilitiesEditor/Private/VPUtilitiesEditorModule.cpp b/Engine/Plugins/Experimental/VirtualProductionUtilities/Source/VPUtilitiesEditor/Private/VPUtilitiesEditorModule.cpp
--- a/Engine/Plugins/Experimental/VirtualProductionUtilities/Source/VPUtilitiesEditor/Private/VPUtilitiesEditorModule.cpp
+++ b/Engine/Plugins/Experimental/VirtualProductionUtilities/Source/VPUtilitiesEditor/Private/VPUtilitiesEditorModule.cpp
@@ -25,6 +25,7 @@
 #include "VPUtilitiesEditorStyle.h"
 #include "WorkspaceMenuStructure.h"
 #include "WorkspaceMenuStructureModule.h"
+#include "LevelEditorOutlinerSettings.h"
 
 LLM_DEFINE_TAG(VirtualProductionUtilities_VPUtilitiesEditor);
 #define LOCTEXT_NAMESPACE "VPUtilitiesEditor"
@@ -32,7 +33,6 @@ LLM_DEFINE_TAG(VirtualProductionUtilities_VPUtilitiesEditor);
 DEFINE_LOG_CATEGORY(LogVPUtilitiesEditor);
 
 const FName FVPUtilitiesEditorModule::VPRoleNotificationBarIdentifier = TEXT("VPRoles");
-const FName FVPUtilitiesEditorModule::PlacementModeCategoryHandle = TEXT("VirtualProduction");
 
 void FVPUtilitiesEditorModule::StartupModule()
 {
@@ -83,7 +83,7 @@ const FPlacementCategoryInfo* FVPUtilitiesEditorModule::GetVirtualProductionPlac
 	{
 		IPlacementModeModule& PlacmentModeModule = IPlacementModeModule::Get();
 
-		if (const FPlacementCategoryInfo* RegisteredInfo = PlacmentModeModule.GetRegisteredPlacementCategory(PlacementModeCategoryHandle))
+		if (const FPlacementCategoryInfo* RegisteredInfo = PlacmentModeModule.GetRegisteredPlacementCategory(FLevelEditorOutlinerBuiltInCategories::VirtualProduction()))
 		{
 			return RegisteredInfo;
 		}
@@ -92,7 +92,7 @@ const FPlacementCategoryInfo* FVPUtilitiesEditorModule::GetVirtualProductionPlac
 			FPlacementCategoryInfo Info(
 				LOCTEXT("VirtualProductionCategoryName", "Virtual Production"),
 				FSlateIcon(FVPUtilitiesEditorStyle::GetStyleSetName(), "PlacementBrowser.Icons.VirtualProduction"),
-				PlacementModeCategoryHandle,
+				FLevelEditorOutlinerBuiltInCategories::VirtualProduction(),
 				TEXT("PMVirtualProduction"),
 				25
 			);
@@ -101,7 +101,7 @@ const FPlacementCategoryInfo* FVPUtilitiesEditorModule::GetVirtualProductionPlac
 
 			// This will return nullptr if the Register above failed so we don't need to explicitly check
 			// RegisterPlacementCategory's return value.
-			return PlacmentModeModule.GetRegisteredPlacementCategory(PlacementModeCategoryHandle);
+			return PlacmentModeModule.GetRegisteredPlacementCategory(FLevelEditorOutlinerBuiltInCategories::VirtualProduction());
 		}
 	}
 
diff --git a/Engine/Plugins/Experimental/VirtualProductionUtilities/Source/VPUtilitiesEditor/Private/VPUtilitiesEditorModule.h b/Engine/Plugins/Experimental/VirtualProductionUtilities/Source/VPUtilitiesEditor/Private/VPUtilitiesEditorModule.h
--- a/Engine/Plugins/Experimental/VirtualProductionUtilities/Source/VPUtilitiesEditor/Private/VPUtilitiesEditorModule.h
+++ b/Engine/Plugins/Experimental/VirtualProductionUtilities/Source/VPUtilitiesEditor/Private/VPUtilitiesEditorModule.h
@@ -52,7 +52,4 @@ private:
 
 	/** UI Handler for virtual scouting. */
 	TStrongObjectPtr<UVPCustomUIHandler> CustomUIHandler;
-
-	/** Unique Handle for the Virtual Production Placement Mode Category */
-	static const FName PlacementModeCategoryHandle;
 };
\ No newline at end of file
diff --git a/Engine/Plugins/VirtualProduction/CameraCalibration/Source/CameraCalibrationEditor/Private/CameraCalibrationEditorModule.cpp b/Engine/Plugins/VirtualProduction/CameraCalibration/Source/CameraCalibrationEditor/Private/CameraCalibrationEditorModule.cpp
--- a/Engine/Plugins/VirtualProduction/CameraCalibration/Source/CameraCalibrationEditor/Private/CameraCalibrationEditorModule.cpp
+++ b/Engine/Plugins/VirtualProduction/CameraCalibration/Source/CameraCalibrationEditor/Private/CameraCalibrationEditorModule.cpp
@@ -24,7 +24,7 @@
 #include "UI/CameraCalibrationMenuEntry.h"
 #include "WorkspaceMenuStructure.h"
 #include "WorkspaceMenuStructureModule.h"
-
+#include "LevelEditorOutlinerSettings.h"
 
 #define LOCTEXT_NAMESPACE "CameraCalibrationEditor"
 
@@ -72,10 +72,8 @@ void FCameraCalibrationEditorModule::StartupModule()
 const FPlacementCategoryInfo* FCameraCalibrationEditorModule::GetVirtualProductionCategoryRegisteredInfo() const
 {
 	IPlacementModeModule& PlacementModeModule = IPlacementModeModule::Get();
-
-	static const FName VirtualProductionName = TEXT("VirtualProduction");
-
-	if (const FPlacementCategoryInfo* RegisteredInfo = PlacementModeModule.GetRegisteredPlacementCategory(VirtualProductionName))
+	
+	if (const FPlacementCategoryInfo* RegisteredInfo = PlacementModeModule.GetRegisteredPlacementCategory(FLevelEditorOutlinerBuiltInCategories::VirtualProduction()))
 	{
 		return RegisteredInfo;
 	}
@@ -84,14 +82,14 @@ const FPlacementCategoryInfo* FCameraCalibrationEditorModule::GetVirtualProducti
 		FPlacementCategoryInfo Info(
 			LOCTEXT("VirtualProductionCategoryName", "Virtual Production"),
 			FSlateIcon(FCameraCalibrationEditorStyle::Get().GetStyleSetName(), "PlacementBrowser.Icons.VirtualProduction"),
-			VirtualProductionName,
+			FLevelEditorOutlinerBuiltInCategories::VirtualProduction(),
 			TEXT("PMVirtualProduction"),
 			25
 		);
 
 		IPlacementModeModule::Get().RegisterPlacementCategory(Info);
 
-		return PlacementModeModule.GetRegisteredPlacementCategory(VirtualProductionName);
+		return PlacementModeModule.GetRegisteredPlacementCategory(FLevelEditorOutlinerBuiltInCategories::VirtualProduction());
 	}
 }
 
diff --git a/Engine/Plugins/VirtualProduction/CameraCalibrationCore/Source/CameraCalibrationCoreEditor/CameraCalibrationCoreEditor.Build.cs b/Engine/Plugins/VirtualProduction/CameraCalibrationCore/Source/CameraCalibrationCoreEditor/CameraCalibrationCoreEditor.Build.cs
--- a/Engine/Plugins/VirtualProduction/CameraCalibrationCore/Source/CameraCalibrationCoreEditor/CameraCalibrationCoreEditor.Build.cs
+++ b/Engine/Plugins/VirtualProduction/CameraCalibrationCore/Source/CameraCalibrationCoreEditor/CameraCalibrationCoreEditor.Build.cs
@@ -19,6 +19,12 @@ namespace UnrealBuildTool.Rules
 					"SlateCore",
 					"UnrealEd",
 				});
+			
+			PrivateIncludePathModuleNames.AddRange(
+				new string[]
+				{
+					"LevelEditor",
+				});
 		}
 	}
 }
diff --git a/Engine/Plugins/VirtualProduction/CameraCalibrationCore/Source/CameraCalibrationCoreEditor/Private/CameraCalibrationCoreEditorModule.cpp b/Engine/Plugins/VirtualProduction/CameraCalibrationCore/Source/CameraCalibrationCoreEditor/Private/CameraCalibrationCoreEditorModule.cpp
--- a/Engine/Plugins/VirtualProduction/CameraCalibrationCore/Source/CameraCalibrationCoreEditor/Private/CameraCalibrationCoreEditorModule.cpp
+++ b/Engine/Plugins/VirtualProduction/CameraCalibrationCore/Source/CameraCalibrationCoreEditor/Private/CameraCalibrationCoreEditorModule.cpp
@@ -16,6 +16,7 @@
 #include "Misc/CoreDelegates.h"
 #include "Modules/ModuleManager.h"
 #include "PropertyEditorModule.h"
+#include "LevelEditorOutlinerSettings.h"
 
 LLM_DEFINE_TAG(CameraCalibrationCore_CameraCalibrationCoreEditor);
 #define LOCTEXT_NAMESPACE "CameraCalibrationCoreEditor"
@@ -73,9 +74,7 @@ const FPlacementCategoryInfo* FCameraCalibrationCoreEditorModule::GetVirtualProd
 {
 	IPlacementModeModule& PlacementModeModule = IPlacementModeModule::Get();
 
-	static const FName VirtualProductionName = TEXT("VirtualProduction");
-
-	if (const FPlacementCategoryInfo* RegisteredInfo = PlacementModeModule.GetRegisteredPlacementCategory(VirtualProductionName))
+	if (const FPlacementCategoryInfo* RegisteredInfo = PlacementModeModule.GetRegisteredPlacementCategory(FLevelEditorOutlinerBuiltInCategories::VirtualProduction()))
 	{
 		return RegisteredInfo;
 	}
@@ -84,14 +83,14 @@ const FPlacementCategoryInfo* FCameraCalibrationCoreEditorModule::GetVirtualProd
 		FPlacementCategoryInfo Info(
 			LOCTEXT("VirtualProductionCategoryName", "Virtual Production"),
 			FSlateIcon(FCameraCalibrationCoreEditorStyle::Get().GetStyleSetName(), "PlacementBrowser.Icons.VirtualProduction"),
-			VirtualProductionName,
+			FLevelEditorOutlinerBuiltInCategories::VirtualProduction(),
 			TEXT("PMVirtualProduction"),
 			25 // Determines where the category shows up in the list with respect to the others.
 		);
 
 		IPlacementModeModule::Get().RegisterPlacementCategory(Info);
 
-		return PlacementModeModule.GetRegisteredPlacementCategory(VirtualProductionName);
+		return PlacementModeModule.GetRegisteredPlacementCategory(FLevelEditorOutlinerBuiltInCategories::VirtualProduction());
 	}
 }
 
diff --git a/Engine/Source/Editor/LevelEditor/Public/LevelEditorOutlinerSettings.h b/Engine/Source/Editor/LevelEditor/Public/LevelEditorOutlinerSettings.h
--- a/Engine/Source/Editor/LevelEditor/Public/LevelEditorOutlinerSettings.h
+++ b/Engine/Source/Editor/LevelEditor/Public/LevelEditorOutlinerSettings.h
@@ -19,16 +19,16 @@ class FCustomClassFilterData;
 /** Structure of built-in filter categories. Defined as functions to enable external use without linkage */
 struct FLevelEditorOutlinerBuiltInCategories
 {
-	static FName Common()			{ static FName Name("Essential");		return Name; }
-	static FName Basic()			 { static FName Name("Basic");			return Name; }
-	static FName Animation()		 { static FName Name("Animation");		return Name; }
-	static FName Audio()			 { static FName Name("Audio");			return Name; }
-	static FName Geometry()			 { static FName Name("Geometry");		return Name; }
-	static FName Lights()			 { static FName Name("Lights");			return Name; }
-	static FName Environment()		 { static FName Name("Environment");	return Name; }
-	static FName Visual()			 { static FName Name("Visual");			return Name; }
-	static FName Volumes()			 { static FName Name("Volumes");		return Name; }
-	static FName VirtualProduction() { static FName Name("VP");				return Name; }
+	static FName Common()			 { static FName Name("Essential");			return Name; }
+	static FName Basic()			 { static FName Name("Basic");				return Name; }
+	static FName Animation()		 { static FName Name("Animation");			return Name; }
+	static FName Audio()			 { static FName Name("Audio");				return Name; }
+	static FName Geometry()			 { static FName Name("Geometry");			return Name; }
+	static FName Lights()			 { static FName Name("Lights");				return Name; }
+	static FName Environment()		 { static FName Name("Environment");		return Name; }
+	static FName Visual()			 { static FName Name("Visual");				return Name; }
+	static FName Volumes()			 { static FName Name("Volumes");			return Name; }
+	static FName VirtualProduction() { static FName Name("VirtualProduction");	return Name; }
 };
 
 /** Helper class to manage initalization options specific to the Level Editor Outliners
