Keep track of remapped widgets across enabling & disabling across same session

diff --git a/Engine/Plugins/Experimental/VirtualCamera/Source/VCamCore/Private/Output/VCamOutputProviderBase.cpp b/Engine/Plugins/Experimental/VirtualCamera/Source/VCamCore/Private/Output/VCamOutputProviderBase.cpp
--- a/Engine/Plugins/Experimental/VirtualCamera/Source/VCamCore/Private/Output/VCamOutputProviderBase.cpp
+++ b/Engine/Plugins/Experimental/VirtualCamera/Source/VCamCore/Private/Output/VCamOutputProviderBase.cpp
@@ -231,7 +231,16 @@ void UVCamOutputProviderBase::DisplayUMG()
 
 		if (ActorWorld)
 		{
+#if WITH_EDITOR
+			if (SavedConnectionRemappingData && SavedConnectionRemappingData->GetClass() == UMGClass)
+			{
+				UMGWidget->SetOverrideWidget(SavedConnectionRemappingData);
+			}
 			UMGWidget->Display(ActorWorld);
+			SavedConnectionRemappingData = UMGWidget->GetWidget();
+#else
+			UMGWidget->Display(ActorWorld);
+#endif
 			UE_LOG(LogVCamOutputProvider, Log, TEXT("DisplayUMG widget displayed in WorldType %d"), WorldType);
 		}
 
diff --git a/Engine/Plugins/Experimental/VirtualCamera/Source/VCamCore/Public/Output/VCamOutputProviderBase.h b/Engine/Plugins/Experimental/VirtualCamera/Source/VCamCore/Public/Output/VCamOutputProviderBase.h
--- a/Engine/Plugins/Experimental/VirtualCamera/Source/VCamCore/Public/Output/VCamOutputProviderBase.h
+++ b/Engine/Plugins/Experimental/VirtualCamera/Source/VCamCore/Public/Output/VCamOutputProviderBase.h
@@ -182,6 +182,12 @@ private:
 	
 	UPROPERTY(Transient)
 	bool bWasActive = false;
+
+#if WITH_EDITORONLY_DATA
+	/** Used for saving widget remapping settings across enabling & disabling in the same session, see FOutputProviderLayoutCustomization. */
+	UPROPERTY(Transient)
+	TObjectPtr<UUserWidget> SavedConnectionRemappingData = nullptr;
+#endif
 	
 	bool IsOuterComponentEnabled() const;
 };
diff --git a/Engine/Plugins/Experimental/VirtualProductionUtilities/Source/VPUtilities/Private/VPFullScreenUserWidget.cpp b/Engine/Plugins/Experimental/VirtualProductionUtilities/Source/VPUtilities/Private/VPFullScreenUserWidget.cpp
--- a/Engine/Plugins/Experimental/VirtualProductionUtilities/Source/VPUtilities/Private/VPFullScreenUserWidget.cpp
+++ b/Engine/Plugins/Experimental/VirtualProductionUtilities/Source/VPUtilities/Private/VPFullScreenUserWidget.cpp
@@ -875,6 +875,14 @@ void UVPFullScreenUserWidget::SetDisplayTypes(EVPWidgetDisplayType InEditorDispl
 	PIEDisplayType = InPIEDisplayType;
 }
 
+void UVPFullScreenUserWidget::SetOverrideWidget(UUserWidget* InWidget)
+{
+	if (ensureMsgf(!IsDisplayed(), TEXT("For simplicity of API you can only override the widget before displaying.")))
+	{
+		Widget = InWidget;
+	}
+}
+
 void UVPFullScreenUserWidget::SetCustomPostProcessSettingsSource(TWeakObjectPtr<UObject> InCustomPostProcessSettingsSource)
 {
 	PostProcessDisplayType.SetCustomPostProcessSettingsSource(InCustomPostProcessSettingsSource);
diff --git a/Engine/Plugins/Experimental/VirtualProductionUtilities/Source/VPUtilities/Public/VPFullScreenUserWidget.h b/Engine/Plugins/Experimental/VirtualProductionUtilities/Source/VPUtilities/Public/VPFullScreenUserWidget.h
--- a/Engine/Plugins/Experimental/VirtualProductionUtilities/Source/VPUtilities/Public/VPFullScreenUserWidget.h
+++ b/Engine/Plugins/Experimental/VirtualProductionUtilities/Source/VPUtilities/Public/VPFullScreenUserWidget.h
@@ -219,6 +219,8 @@ public:
 
 	void SetDisplayTypes(EVPWidgetDisplayType InEditorDisplayType, EVPWidgetDisplayType InGameDisplayType, EVPWidgetDisplayType InPIEDisplayType);
 	
+	void SetOverrideWidget(UUserWidget* InWidget);
+	
 	/**
 	 * If using EVPWidgetDisplayType::PostProcess, you can specify a custom post process settings that should be modified.
 	 * By default, a new post process component is added to AWorldSettings.
