Ensure VCam HUDS are always a higher blend priority than post process settings on the camera

diff --git a/Engine/Plugins/Experimental/VirtualProductionUtilities/Source/VPUtilities/Private/VPFullScreenUserWidget.cpp b/Engine/Plugins/Experimental/VirtualProductionUtilities/Source/VPUtilities/Private/VPFullScreenUserWidget.cpp
--- a/Engine/Plugins/Experimental/VirtualProductionUtilities/Source/VPUtilities/Private/VPFullScreenUserWidget.cpp
+++ b/Engine/Plugins/Experimental/VirtualProductionUtilities/Source/VPUtilities/Private/VPFullScreenUserWidget.cpp
@@ -390,9 +390,9 @@ bool FVPFullScreenUserWidget_PostProcess::InitPostProcessComponent(UWorld* World
 
 		if (FPostProcessSettings* const PostProcessSettings = GetPostProcessSettings())
 		{
-			PostProcessSettings->WeightedBlendables.Array.SetNumZeroed(1);
-			PostProcessSettings->WeightedBlendables.Array[0].Weight = 1.f;
-			PostProcessSettings->WeightedBlendables.Array[0].Object = PostProcessMaterialInstance;
+			const FWeightedBlendable Blendable{ 1.f, PostProcessMaterialInstance };
+			// Pre-existing blend material should not affect the widget so insert the material at the beginning
+			PostProcessSettings->WeightedBlendables.Array.Insert(Blendable, 0);
 			return true;
 		}
 	}
