Post process output providers, such as Pixel Streaming, will no longer save material instances into post process volumes. This fixes there being a white screen on the camera after saving and reloading a map.

diff --git a/Engine/Plugins/Experimental/VirtualProductionUtilities/Source/VPUtilities/Private/VPFullScreenUserWidget.cpp b/Engine/Plugins/Experimental/VirtualProductionUtilities/Source/VPUtilities/Private/VPFullScreenUserWidget.cpp
--- a/Engine/Plugins/Experimental/VirtualProductionUtilities/Source/VPUtilities/Private/VPFullScreenUserWidget.cpp
+++ b/Engine/Plugins/Experimental/VirtualProductionUtilities/Source/VPUtilities/Private/VPFullScreenUserWidget.cpp
@@ -338,6 +338,7 @@ TSharedPtr<SVirtualWindow> FVPFullScreenUserWidget_PostProcess::GetSlateWindow()
 	return SlateWindow;
 }
 
+
 bool FVPFullScreenUserWidget_PostProcess::InitPostProcessComponent(UWorld* World)
 {
 	ReleasePostProcessComponent();
@@ -352,25 +353,47 @@ bool FVPFullScreenUserWidget_PostProcess::InitPostProcessComponent(UWorld* World
 			PostProcessComponent->bUnbound = true;
 			PostProcessComponent->RegisterComponent();
 		}
-		
-		PostProcessMaterialInstance = UMaterialInstanceDynamic::Create(PostProcessMaterial, World);
-		if (!ensure(PostProcessMaterialInstance))
-		{
-			return false;
-		}
 
-		// set the parameter immediately
-		PostProcessMaterialInstance->SetTextureParameterValue(NAME_SlateUI, WidgetRenderTarget);
-		PostProcessMaterialInstance->SetVectorParameterValue(NAME_TintColorAndOpacity, PostProcessTintColorAndOpacity);
-		PostProcessMaterialInstance->SetScalarParameterValue(NAME_OpacityFromTexture, PostProcessOpacityFromTexture);
+		return InitPostProcessMaterial();
+	}
+
+	return false;
+}
 
-		if (FPostProcessSettings* const PostProcessSettings = GetPostProcessSettings())
+bool FVPFullScreenUserWidget_PostProcess::InitPostProcessMaterial()
+{
+	// Passing in a nullptr outer will cause the material to be transient.
+	constexpr UObject* MaterialOuter = nullptr;
+	PostProcessMaterialInstance = UMaterialInstanceDynamic::Create(PostProcessMaterial, MaterialOuter);
+	if (!ensure(PostProcessMaterialInstance))
+	{
+		return false;
+	}
+
+	// set the parameter immediately
+	PostProcessMaterialInstance->SetTextureParameterValue(NAME_SlateUI, WidgetRenderTarget);
+	PostProcessMaterialInstance->SetVectorParameterValue(NAME_TintColorAndOpacity, PostProcessTintColorAndOpacity);
+	PostProcessMaterialInstance->SetScalarParameterValue(NAME_OpacityFromTexture, PostProcessOpacityFromTexture);
+
+	if (FPostProcessSettings* const PostProcessSettings = GetPostProcessSettings())
+	{
+		// User added blend material should not affect the widget so insert the material at the beginning
+		const FWeightedBlendable Blendable{ 1.f, PostProcessMaterialInstance };
+			
+		// Use case: Virtual Camera specifies an external post process settings
+		// 1. Virtual Camera is activated > creates this widget
+		// 2. Save Map
+		// 3. Reload map > we'll have an empty slot in the blendables because PostProcessMaterialInstance is transient
+		const bool bReuseOldEmptySlot = PostProcessSettings->WeightedBlendables.Array.Num() > 0 && !PostProcessSettings->WeightedBlendables.Array[0].Object;
+		if (bReuseOldEmptySlot)
+		{
+			PostProcessSettings->WeightedBlendables.Array[0].Object = PostProcessMaterialInstance;
+		}
+		else
 		{
-			const FWeightedBlendable Blendable{ 1.f, PostProcessMaterialInstance };
-			// Pre-existing blend material should not affect the widget so insert the material at the beginning
 			PostProcessSettings->WeightedBlendables.Array.Insert(Blendable, 0);
-			return true;
 		}
+		return true;
 	}
 
 	return false;
@@ -837,10 +860,6 @@ void UVPFullScreenUserWidget::Hide()
 	if (CurrentDisplayType != EVPWidgetDisplayType::Inactive)
 	{
 		ReleaseWidget();
-		FWorldDelegates::LevelRemovedFromWorld.RemoveAll(this);
-		FWorldDelegates::OnWorldCleanup.RemoveAll(this);
-
-		VPFullScreenUserWidgetPrivate::FWorldCleanupListener::Get()->RemoveWidget(this);
 
 		if (CurrentDisplayType == EVPWidgetDisplayType::Viewport)
 		{
@@ -853,6 +872,9 @@ void UVPFullScreenUserWidget::Hide()
 		CurrentDisplayType = EVPWidgetDisplayType::Inactive;
 	}
 
+	FWorldDelegates::LevelRemovedFromWorld.RemoveAll(this);
+	FWorldDelegates::OnWorldCleanup.RemoveAll(this);
+	VPFullScreenUserWidgetPrivate::FWorldCleanupListener::Get()->RemoveWidget(this);
 	World.Reset();
 }
 
diff --git a/Engine/Plugins/Experimental/VirtualProductionUtilities/Source/VPUtilities/Public/VPFullScreenUserWidget.h b/Engine/Plugins/Experimental/VirtualProductionUtilities/Source/VPUtilities/Public/VPFullScreenUserWidget.h
--- a/Engine/Plugins/Experimental/VirtualProductionUtilities/Source/VPUtilities/Public/VPFullScreenUserWidget.h
+++ b/Engine/Plugins/Experimental/VirtualProductionUtilities/Source/VPUtilities/Public/VPFullScreenUserWidget.h
@@ -85,6 +85,7 @@ struct FVPFullScreenUserWidget_PostProcess
 
 private:
 	bool InitPostProcessComponent(UWorld* World);
+	bool InitPostProcessMaterial();
 	void ReleasePostProcessComponent();
 
 	bool CreateRenderer(UWorld* World, UUserWidget* Widget, TAttribute<float> InDPIScale);
