Fix viewport being white after restarting the editor if Pixel Streaming enable flag was saved into the map

diff --git a/Engine/Plugins/Experimental/VirtualProductionUtilities/Source/VPUtilities/Private/VPFullScreenUserWidget.cpp b/Engine/Plugins/Experimental/VirtualProductionUtilities/Source/VPUtilities/Private/VPFullScreenUserWidget.cpp
--- a/Engine/Plugins/Experimental/VirtualProductionUtilities/Source/VPUtilities/Private/VPFullScreenUserWidget.cpp
+++ b/Engine/Plugins/Experimental/VirtualProductionUtilities/Source/VPUtilities/Private/VPFullScreenUserWidget.cpp
@@ -433,60 +433,53 @@ bool FVPFullScreenUserWidget_PostProcess::CreateRenderer(UWorld* World, UUserWid
 
 	if (World && Widget)
 	{
-		const FIntPoint CalculatedWidgetSize = CalculateWidgetDrawSize(World);
-		if (IsTextureSizeValid(CalculatedWidgetSize))
-		{
-			CurrentWidgetDrawSize = CalculatedWidgetSize;
-
-			const bool bApplyGammaCorrection = true;
-			WidgetRenderer = new FWidgetRenderer(bApplyGammaCorrection);
-			WidgetRenderer->SetIsPrepassNeeded(true);
-
-			SlateWindow = SNew(SVirtualWindow).Size(CurrentWidgetDrawSize);
-			SlateWindow->SetIsFocusable(bWindowFocusable);
-			SlateWindow->SetVisibility(ConvertWindowVisibilityToVisibility(WindowVisibility));
-			SlateWindow->SetContent(
-				SNew(SDPIScaler)
-				.DPIScale(InDPIScale)
-				[
-					Widget->TakeWidget()
-				]
-			);
-
-			RegisterHitTesterWithViewport(World);
-
-			if (!Widget->IsDesignTime() && World->IsGameWorld())
-			{
-				UGameInstance* GameInstance = World->GetGameInstance();
-				UGameViewportClient* GameViewportClient = GameInstance ? GameInstance->GetGameViewportClient() : nullptr;
-				if (GameViewportClient)
-				{
-					SlateWindow->AssignParentWidget(GameViewportClient->GetGameViewportWidget());
-				}
-			}
-
-			FLinearColor ActualBackgroundColor = RenderTargetBackgroundColor;
-			switch (RenderTargetBlendMode)
-			{
-			case EWidgetBlendMode::Opaque:
-				ActualBackgroundColor.A = 1.0f;
-				break;
-			case EWidgetBlendMode::Masked:
-				ActualBackgroundColor.A = 0.0f;
-				break;
-			}
+		constexpr bool bApplyGammaCorrection = true;
+		WidgetRenderer = new FWidgetRenderer(bApplyGammaCorrection);
+		WidgetRenderer->SetIsPrepassNeeded(true);
+		
+		// CalculateWidgetDrawSize may sometimes return {0,0}, e.g. right after engine startup when viewport not yet initialized.
+		// TickRenderer will call Resize automatically once CurrentWidgetDrawSize is updated to be non-zero.
+		checkf(CurrentWidgetDrawSize == FIntPoint::ZeroValue, TEXT("Expected ReleaseRenderer to reset CurrentWidgetDrawSize."));
+		SlateWindow = SNew(SVirtualWindow).Size(CurrentWidgetDrawSize);
+		SlateWindow->SetIsFocusable(bWindowFocusable);
+		SlateWindow->SetVisibility(ConvertWindowVisibilityToVisibility(WindowVisibility));
+		SlateWindow->SetContent(
+			SNew(SDPIScaler)
+			.DPIScale(InDPIScale)
+			[
+				Widget->TakeWidget()
+			]
+		);
 
-			AWorldSettings* WorldSetting = World->GetWorldSettings();
-			WidgetRenderTarget = NewObject<UTextureRenderTarget2D>(WorldSetting, NAME_None, RF_Transient);
-			WidgetRenderTarget->ClearColor = ActualBackgroundColor;
-			WidgetRenderTarget->InitCustomFormat(CurrentWidgetDrawSize.X, CurrentWidgetDrawSize.Y, PF_B8G8R8A8, false);
-			WidgetRenderTarget->UpdateResourceImmediate();
+		RegisterHitTesterWithViewport(World);
 
-			if (!bRenderToTextureOnly && PostProcessMaterialInstance)
+		if (!Widget->IsDesignTime() && World->IsGameWorld())
+		{
+			UGameInstance* GameInstance = World->GetGameInstance();
+			UGameViewportClient* GameViewportClient = GameInstance ? GameInstance->GetGameViewportClient() : nullptr;
+			if (GameViewportClient)
 			{
-				PostProcessMaterialInstance->SetTextureParameterValue(NAME_SlateUI, WidgetRenderTarget);
+				SlateWindow->AssignParentWidget(GameViewportClient->GetGameViewportWidget());
 			}
 		}
+
+		FLinearColor ActualBackgroundColor = RenderTargetBackgroundColor;
+		switch (RenderTargetBlendMode)
+		{
+		case EWidgetBlendMode::Opaque:
+			ActualBackgroundColor.A = 1.0f;
+			break;
+		case EWidgetBlendMode::Masked:
+			ActualBackgroundColor.A = 0.0f;
+			break;
+		}
+
+		// Skip InitCustomFormat call because CalculateWidgetDrawSize may sometimes return {0,0}, e.g. right after engine startup when viewport not yet initialized
+		// TickRenderer will call InitCustomFormat automatically once CurrentWidgetDrawSize is updated to be non-zero.
+		checkf(CurrentWidgetDrawSize == FIntPoint::ZeroValue, TEXT("Expected ReleaseRenderer to reset CurrentWidgetDrawSize."));
+		AWorldSettings* WorldSetting = World->GetWorldSettings();
+		WidgetRenderTarget = NewObject<UTextureRenderTarget2D>(WorldSetting, NAME_None, RF_Transient);
+		WidgetRenderTarget->ClearColor = ActualBackgroundColor;
 	}
 
 	return WidgetRenderer && WidgetRenderTarget;
@@ -533,7 +526,7 @@ void FVPFullScreenUserWidget_PostProcess::TickRenderer(UWorld* World, float Delt
 			}
 		}
 
-		if (WidgetRenderer)
+		if (WidgetRenderer && CurrentWidgetDrawSize != FIntPoint::ZeroValue)
 		{
 			WidgetRenderer->DrawWindow(
 				WidgetRenderTarget,
