FullScreenWidgetActor does not display widget

diff --git a/Engine/Plugins/Experimental/VirtualProductionUtilities/Source/VPUtilities/Private/VPFullScreenUserWidget.cpp b/Engine/Plugins/Experimental/VirtualProductionUtilities/Source/VPUtilities/Private/VPFullScreenUserWidget.cpp
--- a/Engine/Plugins/Experimental/VirtualProductionUtilities/Source/VPUtilities/Private/VPFullScreenUserWidget.cpp
+++ b/Engine/Plugins/Experimental/VirtualProductionUtilities/Source/VPUtilities/Private/VPFullScreenUserWidget.cpp
@@ -779,6 +779,9 @@ bool UVPFullScreenUserWidget::Display(UWorld* InWorld)
 			return false;
 		}
 	}
+	
+	// Make sure that each display type has also received the EditorTargetViewport
+	SetEditorTargetViewport(EditorTargetViewport);
 #endif
 
 	bool bWasAdded = false;
