Fix VCam memory leak using "Save Current Level As" due to creating render target with world as outer instead of as transient package as outer

diff --git a/Engine/Plugins/Experimental/VirtualProductionUtilities/Source/VPUtilities/Private/VPFullScreenUserWidget.cpp b/Engine/Plugins/Experimental/VirtualProductionUtilities/Source/VPUtilities/Private/VPFullScreenUserWidget.cpp
--- a/Engine/Plugins/Experimental/VirtualProductionUtilities/Source/VPUtilities/Private/VPFullScreenUserWidget.cpp
+++ b/Engine/Plugins/Experimental/VirtualProductionUtilities/Source/VPUtilities/Private/VPFullScreenUserWidget.cpp
@@ -362,9 +362,9 @@ bool FVPFullScreenUserWidget_PostProcess::InitPostProcessComponent(UWorld* World
 
 bool FVPFullScreenUserWidget_PostProcess::InitPostProcessMaterial()
 {
-	// Passing in a nullptr outer will cause the material to be transient.
-	constexpr UObject* MaterialOuter = nullptr;
-	PostProcessMaterialInstance = UMaterialInstanceDynamic::Create(PostProcessMaterial, MaterialOuter);
+	// Outer needs to be transient package: otherwise we cause a world memory leak using "Save Current Level As" due to reference not getting replaced correctly
+	PostProcessMaterialInstance = UMaterialInstanceDynamic::Create(PostProcessMaterial, GetTransientPackage());
+	PostProcessMaterialInstance->SetFlags(RF_Transient);
 	if (!ensure(PostProcessMaterialInstance))
 	{
 		return false;
@@ -477,8 +477,8 @@ bool FVPFullScreenUserWidget_PostProcess::CreateRenderer(UWorld* World, UUserWid
 		// Skip InitCustomFormat call because CalculateWidgetDrawSize may sometimes return {0,0}, e.g. right after engine startup when viewport not yet initialized
 		// TickRenderer will call InitCustomFormat automatically once CurrentWidgetDrawSize is updated to be non-zero.
 		checkf(CurrentWidgetDrawSize == FIntPoint::ZeroValue, TEXT("Expected ReleaseRenderer to reset CurrentWidgetDrawSize."));
-		AWorldSettings* WorldSetting = World->GetWorldSettings();
-		WidgetRenderTarget = NewObject<UTextureRenderTarget2D>(WorldSetting, NAME_None, RF_Transient);
+		// Outer needs to be transient package: otherwise we cause a world memory leak using "Save Current Level As" due to reference not getting replaced correctly
+		WidgetRenderTarget = NewObject<UTextureRenderTarget2D>(GetTransientPackage(), NAME_None, RF_Transient);
 		WidgetRenderTarget->ClearColor = ActualBackgroundColor;
 	}
 
