Fix missing include for UPackage

diff --git a/Engine/Plugins/Experimental/VirtualProductionUtilities/Source/VPUtilities/Private/VPFullScreenUserWidget.cpp b/Engine/Plugins/Experimental/VirtualProductionUtilities/Source/VPUtilities/Private/VPFullScreenUserWidget.cpp
--- a/Engine/Plugins/Experimental/VirtualProductionUtilities/Source/VPUtilities/Private/VPFullScreenUserWidget.cpp
+++ b/Engine/Plugins/Experimental/VirtualProductionUtilities/Source/VPUtilities/Private/VPFullScreenUserWidget.cpp
@@ -13,6 +13,7 @@
 #include "RenderingThread.h"
 #include "RHI.h"
 #include "UObject/ConstructorHelpers.h"
+#include "UObject/Package.h"
 #include "VPUtilitiesModule.h"
 #include "HAL/PlatformApplicationMisc.h"
 
