Create a version of UAssetThumbnailWidget which builds for packaged builds

diff --git a/Engine/Plugins/Experimental/VirtualProductionUtilities/Source/VPUtilities/Private/Widgets/VPAssetThumbnailWrapperWidget.cpp b/Engine/Plugins/Experimental/VirtualProductionUtilities/Source/VPUtilities/Private/Widgets/VPAssetThumbnailWrapperWidget.cpp
new file mode 100644
--- /dev/null
+++ b/Engine/Plugins/Experimental/VirtualProductionUtilities/Source/VPUtilities/Private/Widgets/VPAssetThumbnailWrapperWidget.cpp
@@ -0,0 +1,146 @@
+// Copyright Epic Games, Inc. All Rights Reserved.
+
+#include "Widgets/VPAssetThumbnailWrapperWidget.h"
+
+#include "AssetRegistry/IAssetRegistry.h"
+#include "Widgets/Images/SImage.h"
+#include "Widgets/Layout/SWidgetSwitcher.h"
+
+#if WITH_EDITOR
+#include "Components/AssetThumbnailWidget.h"
+#include "ThumbnailRendering/ThumbnailManager.h"
+#endif
+
+#define LOCTEXT_NAMESPACE "UVPAssetThumbnailWrapperWidget"
+
+void UVPAssetThumbnailWrapperWidget::SetEditorThumbnailResolution(const FIntPoint& NewResolution)
+{
+#if WITH_EDITOR
+	if (AssetWidget)
+	{
+		ExactCast<UAssetThumbnailWidget>(AssetWidget)->SetResolution(NewResolution);
+	}
+#endif
+}
+
+FIntPoint UVPAssetThumbnailWrapperWidget::GetEditorThumbnailResolution() const
+{
+#if WITH_EDITOR
+	if (AssetWidget)
+	{
+		return ExactCast<UAssetThumbnailWidget>(AssetWidget)->GetResolution();
+	}
+#endif
+	return FIntPoint::ZeroValue;
+}
+
+UObject* UVPAssetThumbnailWrapperWidget::GetEditorAssetWidget() const
+{
+#if WITH_EDITOR
+	return AssetWidget;
+#else
+	return nullptr;
+#endif
+}
+
+void UVPAssetThumbnailWrapperWidget::SetAsset(const FAssetData& AssetData)
+{
+#if WITH_EDITOR
+	if (AssetWidget)
+	{
+		ExactCast<UAssetThumbnailWidget>(AssetWidget)->SetAsset(AssetData);
+	}
+#endif
+}
+
+void UVPAssetThumbnailWrapperWidget::SetAssetByObject(UObject* Object)
+{
+#if WITH_EDITOR
+	if (AssetWidget)
+	{
+		ExactCast<UAssetThumbnailWidget>(AssetWidget)->SetAssetByObject(Object);
+	}
+#endif
+}
+
+void UVPAssetThumbnailWrapperWidget::SetFallbackBrush(const FSlateBrush& NewFallbackBrush)
+{
+	FallbackBrush = NewFallbackBrush;
+	UpdateNativeAssetThumbnailWidget();
+}
+
+void UVPAssetThumbnailWrapperWidget::SetDisplayMode(EAssetThumbnailDisplayMode Mode)
+{
+#if WITH_EDITOR
+	DisplayMode = Mode;
+	UpdateNativeAssetThumbnailWidget();
+#endif
+}
+
+void UVPAssetThumbnailWrapperWidget::SynchronizeProperties()
+{
+	Super::SynchronizeProperties();
+#if WITH_EDITOR
+	if (AssetWidget)
+	{
+		ExactCast<UAssetThumbnailWidget>(AssetWidget)->SynchronizeProperties();
+	}
+#endif
+	UpdateNativeAssetThumbnailWidget();
+}
+
+void UVPAssetThumbnailWrapperWidget::ReleaseSlateResources(bool bReleaseChildren)
+{
+	Super::ReleaseSlateResources(bReleaseChildren);
+
+#if WITH_EDITOR
+	if (AssetWidget)
+	{
+		ExactCast<UAssetThumbnailWidget>(AssetWidget)->ReleaseSlateResources(bReleaseChildren);
+		AssetWidget = nullptr;
+	}
+#endif
+	Fallback.Reset();
+	Content.Reset();
+}
+
+TSharedRef<SWidget> UVPAssetThumbnailWrapperWidget::RebuildWidget()
+{
+#if WITH_EDITOR
+	AssetWidget = NewObject<UAssetThumbnailWidget>(this);
+#endif
+	
+	Content = SNew(SWidgetSwitcher)
+		.WidgetIndex(0)
+#if WITH_EDITOR
+		+SWidgetSwitcher::Slot()
+		[
+			ExactCast<UAssetThumbnailWidget>(AssetWidget)->RebuildWidget()
+		]
+#endif
+	+SWidgetSwitcher::Slot()
+		[
+			SAssignNew(Fallback, SImage)
+		];
+	
+	UpdateNativeAssetThumbnailWidget();
+	return Content.ToSharedRef();
+}
+
+#if WITH_EDITOR
+const FText UVPAssetThumbnailWrapperWidget::GetPaletteCategory()
+{
+	return LOCTEXT("VirtualProduction", "Virtual Production");
+}
+#endif
+
+void UVPAssetThumbnailWrapperWidget::UpdateNativeAssetThumbnailWidget()
+{
+	Fallback->SetImage(&FallbackBrush);
+
+#if WITH_EDITOR
+	Content->SetActiveWidgetIndex(static_cast<int32>(DisplayMode));
+#endif
+}
+
+#undef LOCTEXT_NAMESPACE
\ No newline at end of file
diff --git a/Engine/Plugins/Experimental/VirtualProductionUtilities/Source/VPUtilities/Public/Widgets/VPAssetThumbnailWrapperWidget.h b/Engine/Plugins/Experimental/VirtualProductionUtilities/Source/VPUtilities/Public/Widgets/VPAssetThumbnailWrapperWidget.h
new file mode 100644
--- /dev/null
+++ b/Engine/Plugins/Experimental/VirtualProductionUtilities/Source/VPUtilities/Public/Widgets/VPAssetThumbnailWrapperWidget.h
@@ -0,0 +1,92 @@
+// Copyright Epic Games, Inc. All Rights Reserved.
+
+#pragma once
+
+#include "Components/Widget.h"
+#include "VPAssetThumbnailWrapperWidget.generated.h"
+
+class SImage;
+class SWidgetSwitcher;
+
+UENUM()
+enum class EAssetThumbnailDisplayMode : uint8
+{
+	EditorThumbnail,
+	FallbackImage
+};
+
+/**
+ * Version of UAssetThumbnailWidget that compiles in packaged games.
+ * In editor builds, a thumbnail widget is displayed.
+ * In packaged builds, a fallback SImage is displayed.
+ */
+UCLASS(DisplayName = "Asset Thumbnail Widget (Editor & Game)")
+class VPUTILITIES_API UVPAssetThumbnailWrapperWidget : public UWidget
+{
+	GENERATED_BODY()
+public:
+
+	/** Sets the resolution of the editor thumbnail */
+	UFUNCTION(BlueprintCallable, Category = "Appearance")
+	void SetEditorThumbnailResolution(const FIntPoint& NewResolution);
+	UFUNCTION(BlueprintPure, Category = "Appearance")
+	FIntPoint GetEditorThumbnailResolution() const;
+
+	/**
+	 * Gets the widget used for displaying in editor. Returns nullptr in non-editor builds.
+	 * This type must be cast to UAssetThumbnailWidget manually (due to Unreal Header Tool restrictions).
+	 */
+	UFUNCTION(BlueprintPure, Category = "Appearance")
+	UObject* GetEditorAssetWidget() const;
+	
+	/** Sets the asset to display. Has no effect in non-editor builds. */
+	UFUNCTION(BlueprintCallable, Category = "Appearance")
+	void SetAsset(const FAssetData& AssetData);
+	/** Sets the asset to display. Has no effect in non-editor builds. */
+	UFUNCTION(BlueprintCallable, Category = "Appearance")
+	void SetAssetByObject(UObject* Object);
+
+	/** Sets the fallback image to display. */
+	UFUNCTION(BlueprintCallable, Category = "Appearance")
+	void SetFallbackBrush(const FSlateBrush& NewFallbackBrush);
+	UFUNCTION(BlueprintGetter, Category = "Appearance")
+	FSlateBrush GetFallbackBrush() const { return FallbackBrush; }
+
+	/** Sets the display mode. Only has an effect in editor builds because packaged games always display in FallbackImage mode. */
+	UFUNCTION(BlueprintCallable, Category = "Appearance")
+	void SetDisplayMode(EAssetThumbnailDisplayMode Mode = EAssetThumbnailDisplayMode::FallbackImage);
+
+	//~ Begin UWidget Interface
+	virtual void SynchronizeProperties() override;
+	virtual void ReleaseSlateResources(bool bReleaseChildren) override;
+	virtual TSharedRef<SWidget> RebuildWidget() override;
+#if WITH_EDITOR
+	virtual const FText GetPaletteCategory() override;
+#endif
+	//~ End UWidget Interface
+	
+private:
+
+	/** Fallback brush to draw */
+	UPROPERTY(EditAnywhere, BlueprintGetter = GetFallbackBrush, BlueprintSetter = SetFallbackBrush, Category = "Appearance")
+	FSlateBrush FallbackBrush;
+
+#if WITH_EDITORONLY_DATA
+	/** The content for the editor version. Is an instance of UAssetThumbnailWidget but for non-editor builds UHT does not find the type so UObject here. */
+	UPROPERTY(Instanced)
+	TObjectPtr<UObject> AssetWidget;
+	
+	/** How the widget is supposed to be displayed. */
+	UPROPERTY(EditAnywhere, BlueprintSetter = SetDisplayMode, AdvancedDisplay, Category = "Appearance")
+	EAssetThumbnailDisplayMode DisplayMode = EAssetThumbnailDisplayMode::EditorThumbnail;
+#endif
+	
+	/** The fallback image */
+	TSharedPtr<SImage> Fallback;
+
+	/** Holds the contents: either ThumbnailRenderer or a fallback SImage.*/
+	TSharedPtr<SWidgetSwitcher> Content;
+	
+	/** Updates or creates ThumbnailRenderer and sets the contents of DisplayedWidget. */
+	void UpdateNativeAssetThumbnailWidget();
+};
diff --git a/Engine/Plugins/Experimental/VirtualProductionUtilities/Source/VPUtilities/VPUtilities.Build.cs b/Engine/Plugins/Experimental/VirtualProductionUtilities/Source/VPUtilities/VPUtilities.Build.cs
--- a/Engine/Plugins/Experimental/VirtualProductionUtilities/Source/VPUtilities/VPUtilities.Build.cs
+++ b/Engine/Plugins/Experimental/VirtualProductionUtilities/Source/VPUtilities/VPUtilities.Build.cs
@@ -45,6 +45,7 @@ public class VPUtilities : ModuleRules
 					"Blutility",
 					"EditorFramework",
 					"LevelEditor",
+					"UMGEditor",
 					"UnrealEd",
 					"ViewportInteraction",
 					"VPBookmarkEditor",
diff --git a/Engine/Source/Editor/UMGEditor/Private/Components/AssetThumbnailWidget.cpp b/Engine/Source/Editor/UMGEditor/Private/Components/AssetThumbnailWidget.cpp
--- a/Engine/Source/Editor/UMGEditor/Private/Components/AssetThumbnailWidget.cpp
+++ b/Engine/Source/Editor/UMGEditor/Private/Components/AssetThumbnailWidget.cpp
@@ -48,6 +48,12 @@ void UAssetThumbnailWidget::SetAssetByObject(UObject* Object)
 	UpdateNativeAssetThumbnailWidget();
 }
 
+void UAssetThumbnailWidget::SetResolution(const FIntPoint& InResolution)
+{
+	Resolution = InResolution;
+	UpdateNativeAssetThumbnailWidget();
+}
+
 void UAssetThumbnailWidget::SynchronizeProperties()
 {
 	Super::SynchronizeProperties();
@@ -79,17 +85,22 @@ const FText UAssetThumbnailWidget::GetPaletteCategory()
 
 void UAssetThumbnailWidget::UpdateNativeAssetThumbnailWidget()
 {
-	if (ThumbnailRenderer && (AssetToShow.IsValid() && ThumbnailRenderer->GetAssetData() != AssetToShow))
+	if (!DisplayedWidget)
 	{
-		ThumbnailRenderer->SetAsset(AssetToShow);
+		return;
 	}
-	else if (AssetToShow.IsValid() &&
-		(!ThumbnailRenderer || ThumbnailRenderer->GetSize() != Resolution))
+
+	const bool bCreateNewThumbnailRenderer = AssetToShow.IsValid() && (!ThumbnailRenderer || ThumbnailRenderer->GetSize() != Resolution);
+	if (bCreateNewThumbnailRenderer)
 	{
 		ThumbnailRenderer = MakeShared<FAssetThumbnail>(AssetToShow, FMath::Clamp(Resolution.X, 1, 1024), FMath::Clamp(Resolution.Y, 1, 1024), UThumbnailManager::Get().GetSharedThumbnailPool());
 	}
+	else if (ThumbnailRenderer && (AssetToShow.IsValid() && ThumbnailRenderer->GetAssetData() != AssetToShow))
+	{
+		ThumbnailRenderer->SetAsset(AssetToShow);
+	}
 	
-	if (ThumbnailRenderer && ensure(DisplayedWidget))
+	if (ThumbnailRenderer)
 	{
 		DisplayedWidget->SetContent(ThumbnailRenderer->MakeThumbnailWidget(ThumbnailSettings.ToThumbnailConfig()));
 	}
diff --git a/Engine/Source/Editor/UMGEditor/Public/Components/AssetThumbnailWidget.h b/Engine/Source/Editor/UMGEditor/Public/Components/AssetThumbnailWidget.h
--- a/Engine/Source/Editor/UMGEditor/Public/Components/AssetThumbnailWidget.h
+++ b/Engine/Source/Editor/UMGEditor/Public/Components/AssetThumbnailWidget.h
@@ -76,9 +76,7 @@ struct UMGEDITOR_API FAssetThumbnailWidgetSettings
 	FAssetThumbnailConfig ToThumbnailConfig() const;
 };
 
-/**
- * 
- */
+/** This widget can be given an asset and it will render its thumbnail. Editor-only. */
 UCLASS()
 class UMGEDITOR_API UAssetThumbnailWidget : public UWidget
 {
@@ -93,6 +91,14 @@ public:
 	void SetAsset(const FAssetData& AssetData);
 	UFUNCTION(BlueprintCallable, Category = "Appearance")
 	void SetAssetByObject(UObject* Object);
+	
+	/** Sets the resolution for the rendered thumbnail. */
+	UFUNCTION(Blueprintcallable, Category = "Appearance")
+	void SetResolution(const FIntPoint& InResolution);
+
+	/** Gets the resolution of the rendered thumbnail. */
+	UFUNCTION(BlueprintPure, Category = "Appearance")
+	FIntPoint GetResolution() const { return Resolution; }
 
 	//~ Begin UWidget Interface
 	virtual void SynchronizeProperties() override;
@@ -108,7 +114,7 @@ private:
 	FAssetData AssetToShow;
 
 	/** Desired size of the thumbnail */
-	UPROPERTY(EditAnywhere, Category = "Appearance", meta = (ClampMin = "1", ClampMax = "1024"))
+	UPROPERTY(EditAnywhere, BlueprintGetter = "GetResolution", BlueprintSetter = "SetResolution", Category = "Appearance", meta = (ClampMin = "1", ClampMax = "1024"))
 	FIntPoint Resolution = { 64, 64};
 	
 	/** Behaviour and style of the widget */
