Fix compile error for UnrealGame

diff --git a/Engine/Plugins/Experimental/VirtualProductionUtilities/Source/VPUtilities/Public/Widgets/VPAssetThumbnailWrapperWidget.h b/Engine/Plugins/Experimental/VirtualProductionUtilities/Source/VPUtilities/Public/Widgets/VPAssetThumbnailWrapperWidget.h
--- a/Engine/Plugins/Experimental/VirtualProductionUtilities/Source/VPUtilities/Public/Widgets/VPAssetThumbnailWrapperWidget.h
+++ b/Engine/Plugins/Experimental/VirtualProductionUtilities/Source/VPUtilities/Public/Widgets/VPAssetThumbnailWrapperWidget.h
@@ -2,6 +2,7 @@
 
 #pragma once
 
+#include "AssetRegistry/AssetData.h"
 #include "Components/Widget.h"
 #include "VPAssetThumbnailWrapperWidget.generated.h"
 
diff --git a/Engine/Plugins/Experimental/VirtualProductionUtilities/Source/VPUtilities/VPUtilities.Build.cs b/Engine/Plugins/Experimental/VirtualProductionUtilities/Source/VPUtilities/VPUtilities.Build.cs
--- a/Engine/Plugins/Experimental/VirtualProductionUtilities/Source/VPUtilities/VPUtilities.Build.cs
+++ b/Engine/Plugins/Experimental/VirtualProductionUtilities/Source/VPUtilities/VPUtilities.Build.cs
@@ -11,6 +11,7 @@ public class VPUtilities : ModuleRules
 		PublicDependencyModuleNames.AddRange(
 			new string[]
 			{
+				"AssetRegistry",
 				"CinematicCamera",
 				"Composure",
 				"Core",
