VirtualProductionUtilities: Add plugin dependency on XRBase because of blueprint references to UHeadMountedDisplayFunctionLibrary.

diff --git a/Engine/Plugins/Experimental/VirtualProductionUtilities/VirtualProductionUtilities.uplugin b/Engine/Plugins/Experimental/VirtualProductionUtilities/VirtualProductionUtilities.uplugin
--- a/Engine/Plugins/Experimental/VirtualProductionUtilities/VirtualProductionUtilities.uplugin
+++ b/Engine/Plugins/Experimental/VirtualProductionUtilities/VirtualProductionUtilities.uplugin
@@ -76,7 +76,10 @@
 		{
 			"Name": "ConsoleVariables",
 			"Enabled": true			
+		},
+		{
+			"Name": "XRBase",
+			"Enabled": true
 		}
-
 	]
 }
