Fixed bug caused by MaterialInstanceDynamic created in a blueprint trough the UI not being marked as public, preventing the level from being saved.

diff --git a/Engine/Plugins/Experimental/VirtualProductionUtilities/Source/VPMaterialsEditor/Private/SMaterialDynamicWidgets.cpp b/Engine/Plugins/Experimental/VirtualProductionUtilities/Source/VPMaterialsEditor/Private/SMaterialDynamicWidgets.cpp
--- a/Engine/Plugins/Experimental/VirtualProductionUtilities/Source/VPMaterialsEditor/Private/SMaterialDynamicWidgets.cpp
+++ b/Engine/Plugins/Experimental/VirtualProductionUtilities/Source/VPMaterialsEditor/Private/SMaterialDynamicWidgets.cpp
@@ -237,7 +237,7 @@ FReply SMaterialDynamicView::OnCreateDynamicMaterialButtonClicked() const
 		FScopedTransaction Transaction(LOCTEXT("CreateDynamicMaterial", "Create Dynamic Material"));
 
 		// Create Transactional Material. That is allawed create MID with MU and transact the changes
-		UMaterialInstanceDynamic* NewMaterialInstanceDynamic = NewObject<UMaterialInstanceDynamic>(OwnerActorComponent, NAME_None, RF_Transactional);
+		UMaterialInstanceDynamic* NewMaterialInstanceDynamic = NewObject<UMaterialInstanceDynamic>(OwnerActorComponent, NAME_None, RF_Public | RF_Transactional);
 		NewMaterialInstanceDynamic->Parent = MaterialInstance;
 		NewMaterialInstanceDynamic->UpdateCachedData();
 
