Fix NonUnity error caused by TArray<T>::DebugGet in certain debug configurations

diff --git a/Engine/Plugins/Experimental/VirtualProductionUtilities/Source/VPUtilities/Public/Widgets/Fullscreen/VPFullScreenUserWidget_PostProcessWithSVE.h b/Engine/Plugins/Experimental/VirtualProductionUtilities/Source/VPUtilities/Public/Widgets/Fullscreen/VPFullScreenUserWidget_PostProcessWithSVE.h
--- a/Engine/Plugins/Experimental/VirtualProductionUtilities/Source/VPUtilities/Public/Widgets/Fullscreen/VPFullScreenUserWidget_PostProcessWithSVE.h
+++ b/Engine/Plugins/Experimental/VirtualProductionUtilities/Source/VPUtilities/Public/Widgets/Fullscreen/VPFullScreenUserWidget_PostProcessWithSVE.h
@@ -3,11 +3,11 @@
 #pragma once
 
 #include "VPFullScreenUserWidget_PostProcessBase.h"
+#include "SceneViewExtensionContext.h"
 #include "VPFullScreenUserWidget_PostProcessWithSVE.generated.h"
 
 class FSceneViewExtensionBase;
 class UMaterialInterface;
-struct FSceneViewExtensionIsActiveFunctor;
 
 /**
  * Renders widget in post process phase by using Scene View Extensions (SVE).
