Fix VPFullScreenWidget not rendering correctly with SVE.
Texture access was not synchronized. The overlaying is now done in a dedicated shader.

diff --git a/Engine/Plugins/Experimental/VirtualProductionUtilities/Shaders/Private/VPFullScreenWidgetOverlay.usf b/Engine/Plugins/Experimental/VirtualProductionUtilities/Shaders/Private/VPFullScreenWidgetOverlay.usf
new file mode 100644
--- /dev/null
+++ b/Engine/Plugins/Experimental/VirtualProductionUtilities/Shaders/Private/VPFullScreenWidgetOverlay.usf
@@ -0,0 +1,30 @@
+// Copyright Epic Games, Inc. All Rights Reserved.
+
+/*=============================================================================
+ VPFullScreenWidgetOverlay.usf: Shader for overlaying a widget over another texture
+ =============================================================================*/
+ 
+#include "/Engine/Public/Platform.ush"
+
+/* Pixel shaders
+=============================================================================*/
+
+// The render target containing the result of FWidgetRenderer::DrawWindow
+Texture2D WidgetTexture;
+SamplerState WidgetSampler;
+// This is the scene color texture, retrieved by the SVE
+Texture2D SceneTexture;
+SamplerState SceneSampler;
+
+void OverlayWidgetPS(
+	noperspective float4 UVAndScreenPos : TEXCOORD0,
+	out float4 OutColor : SV_Target0)
+{
+	float2 UV = UVAndScreenPos.xy;
+	float4 WidgetColor = WidgetTexture.Sample(WidgetSampler, UV);
+	float4 SceneColor = SceneTexture.Sample(SceneSampler, UV);
+	// Blending using the widget texture alpha value - this effectively overlays the widget
+	OutColor.rgba = float4(WidgetColor.rgb * WidgetColor.a + SceneColor.rgb * (1 - WidgetColor.a), SceneColor.a);
+}
+
+
diff --git a/Engine/Plugins/Experimental/VirtualProductionUtilities/Source/VPUtilities/Private/VPUtilitiesModule.cpp b/Engine/Plugins/Experimental/VirtualProductionUtilities/Source/VPUtilities/Private/VPUtilitiesModule.cpp
--- a/Engine/Plugins/Experimental/VirtualProductionUtilities/Source/VPUtilities/Private/VPUtilitiesModule.cpp
+++ b/Engine/Plugins/Experimental/VirtualProductionUtilities/Source/VPUtilities/Private/VPUtilitiesModule.cpp
@@ -6,8 +6,10 @@
 #include "Modules/ModuleManager.h"
 #include "IConcertClient.h"
 
-#include "VPRolesSubsystem.h"
+#include "ShaderCore.h"
 #include "UnrealEngine.h"
+#include "VPRolesSubsystem.h"
+#include "Interfaces/IPluginManager.h"
 
 DEFINE_LOG_CATEGORY(LogVPUtilities);
 
@@ -44,6 +46,9 @@ public:
 	{
 		UE::ConcertClient::VPRoleEvaluator().BindStatic(
 			&UE::VPUtilities::Private::EvaluateRole );
+		
+		const FString PluginShaderDir = FPaths::Combine(IPluginManager::Get().FindPlugin(TEXT("VirtualProductionUtilities"))->GetBaseDir(), TEXT("Shaders"));
+		AddShaderSourceDirectoryMapping(TEXT("/Plugin/VirtualProductionUtilities"), PluginShaderDir);
 	}
 
 	virtual void ShutdownModule() override
diff --git a/Engine/Plugins/Experimental/VirtualProductionUtilities/Source/VPUtilities/Private/Widgets/Fullscreen/PostProcessSceneViewExtension.cpp b/Engine/Plugins/Experimental/VirtualProductionUtilities/Source/VPUtilities/Private/Widgets/Fullscreen/PostProcessSceneViewExtension.cpp
--- a/Engine/Plugins/Experimental/VirtualProductionUtilities/Source/VPUtilities/Private/Widgets/Fullscreen/PostProcessSceneViewExtension.cpp
+++ b/Engine/Plugins/Experimental/VirtualProductionUtilities/Source/VPUtilities/Private/Widgets/Fullscreen/PostProcessSceneViewExtension.cpp
@@ -4,107 +4,19 @@
 
 #include "VPUtilitiesModule.h"
 
-#include "Materials/Material.h"
-#include "Materials/MaterialInterface.h"
-#include "Materials/MaterialRenderProxy.h"
-#include "MaterialDomain.h"
-#include "DataDrivenShaderPlatformInfo.h"
-#include "MaterialShader.h"
-#include "Containers/DynamicRHIResourceArray.h"
 #include "CommonRenderResources.h"
+#include "Engine/TextureRenderTarget2D.h"
+#include "DataDrivenShaderPlatformInfo.h"
 #include "PostProcess/DrawRectangle.h"
 #include "PostProcess/PostProcessMaterialInputs.h"
-#include "SceneRenderTargetParameters.h"
-#include "SceneViewExtension.h"
 #include "ScreenPass.h"
+#include "TextureResource.h"
 
 namespace UE::VirtualProductionUtilities::Private
 {
-	BEGIN_SHADER_PARAMETER_STRUCT(FVPFullScreenPostProcessMaterialParameters, )
-		SHADER_PARAMETER_STRUCT(FScreenPassTextureViewportParameters, PostProcessOutput)
-		SHADER_PARAMETER_STRUCT_ARRAY(FScreenPassTextureInput, PostProcessInput, [kPostProcessMaterialInputCountMax])
-		SHADER_PARAMETER_STRUCT_ARRAY(FScreenPassTextureInput, PathTracingPostProcessInput, [kPathTracingPostProcessMaterialInputCountMax])
-		SHADER_PARAMETER_SAMPLER(SamplerState, PostProcessInput_BilinearSampler)
-		SHADER_PARAMETER_RDG_BUFFER_SRV(StructuredBuffer<float4>, EyeAdaptationBuffer)
-		RENDER_TARGET_BINDING_SLOTS()
-	END_SHADER_PARAMETER_STRUCT()
-
-	class FPostProcessMaterialShader : public FMaterialShader
-	{
-	public:
-		using FParameters = FVPFullScreenPostProcessMaterialParameters;
-		SHADER_USE_PARAMETER_STRUCT_WITH_LEGACY_BASE(FPostProcessMaterialShader, FMaterialShader);
-
-		static bool ShouldCompilePermutation(const FMaterialShaderPermutationParameters& Parameters)
-		{
-			return Parameters.MaterialParameters.MaterialDomain == MD_PostProcess && !IsMobilePlatform(Parameters.Platform);
-		}
-
-		static void ModifyCompilationEnvironment(const FMaterialShaderPermutationParameters& Parameters, FShaderCompilerEnvironment& OutEnvironment)
-		{
-			FMaterialShader::ModifyCompilationEnvironment(Parameters, OutEnvironment);
-			OutEnvironment.SetDefine(TEXT("POST_PROCESS_MATERIAL"), 1);
-			OutEnvironment.SetDefine(TEXT("POST_PROCESS_MATERIAL_BEFORE_TONEMAP"), (Parameters.MaterialParameters.BlendableLocation != BL_AfterTonemapping) ? 1 : 0);
-		}
-	};
-	
-	/** Shaders to render our post process material */
-	class FVPFullScreenPostProcessVS :
-		public FPostProcessMaterialShader
-	{
-	public:
-		DECLARE_SHADER_TYPE(FVPFullScreenPostProcessVS, Material);
-
-		static void ModifyCompilationEnvironment(const FMaterialShaderPermutationParameters& Parameters, FShaderCompilerEnvironment& OutEnvironment)
-		{
-			FPostProcessMaterialShader::ModifyCompilationEnvironment(Parameters, OutEnvironment);
-			OutEnvironment.SetDefine(TEXT("POST_PROCESS_AR_PASSTHROUGH"), 1);
-		}
-
-		FVPFullScreenPostProcessVS() = default;
-		FVPFullScreenPostProcessVS(const ShaderMetaType::CompiledShaderInitializerType& Initializer)
-			: FPostProcessMaterialShader(Initializer)
-		{}
-
-		void SetParameters(FRHIBatchedShaderParameters& BatchedParameters, const FSceneView& View)
-		{
-			UE::Renderer::PostProcess::SetDrawRectangleParameters(BatchedParameters, this, View);
-			FMaterialShader::SetViewParameters(BatchedParameters, View, View.ViewUniformBuffer);
-		}
-	};
-
-	IMPLEMENT_SHADER_TYPE(,FVPFullScreenPostProcessVS, TEXT("/Engine/Private/PostProcessMaterialShaders.usf"), TEXT("MainVS_VideoOverlay"), SF_Vertex);
-
-	class FVPFullScreenPostProcessPS :
-		public FPostProcessMaterialShader
-	{
-	public:
-		DECLARE_SHADER_TYPE(FVPFullScreenPostProcessPS, Material);
-
-		static void ModifyCompilationEnvironment(const FMaterialShaderPermutationParameters& Parameters, FShaderCompilerEnvironment& OutEnvironment)
-		{
-			FPostProcessMaterialShader::ModifyCompilationEnvironment(Parameters, OutEnvironment);
-			OutEnvironment.SetDefine(TEXT("OUTPUT_MOBILE_HDR"), IsMobileHDR() ? 1 : 0);
-		}
-
-		FVPFullScreenPostProcessPS() = default;
-		FVPFullScreenPostProcessPS(const ShaderMetaType::CompiledShaderInitializerType& Initializer)
-			: FPostProcessMaterialShader(Initializer)
-		{}
-
-		void SetParameters(FRHIBatchedShaderParameters& BatchedParameters, const FSceneView& View, const FMaterialRenderProxy* MaterialProxy)
-		{
-			const FMaterial& Material = MaterialProxy->GetMaterialWithFallback(View.GetFeatureLevel(), MaterialProxy);
-			FMaterialShader::SetViewParameters(BatchedParameters, View, View.ViewUniformBuffer);
-			FMaterialShader::SetParameters(BatchedParameters, MaterialProxy, Material, View);
-		}
-	};
-
-	IMPLEMENT_SHADER_TYPE(,FVPFullScreenPostProcessPS, TEXT("/Engine/Private/PostProcessMaterialShaders.usf"), TEXT("MainPS_VideoOverlay"), SF_Pixel);
-	
-	FPostProcessSceneViewExtension::FPostProcessSceneViewExtension(const FAutoRegister& AutoRegister, TAttribute<UMaterialInterface*> PostProcessMaterialGetter)
+	FPostProcessSceneViewExtension::FPostProcessSceneViewExtension(const FAutoRegister& AutoRegister, UTextureRenderTarget2D& WidgetRenderTarget)
 		: Super(AutoRegister)
-		, PostProcessMaterialGetter(MoveTemp(PostProcessMaterialGetter))
+		, WidgetRenderTarget(&WidgetRenderTarget)
 	{}
 
 	void FPostProcessSceneViewExtension::SetupViewFamily(FSceneViewFamily& InViewFamily)
@@ -117,65 +29,14 @@ namespace UE::VirtualProductionUtilities::Private
 	{}
 
 	void FPostProcessSceneViewExtension::PreRenderView_RenderThread(FRDGBuilder& GraphBuilder, FSceneView& InView)
-	{
-		if (VertexBufferRHI == nullptr || !VertexBufferRHI.IsValid())
-		{
-			// Setup vertex buffer
-			TResourceArray<FFilterVertex, VERTEXBUFFER_ALIGNMENT> Vertices;
-			Vertices.SetNumUninitialized(4);
-
-			Vertices[0].Position = FVector4f(0.f, 0.f, 0.f, 1.f);
-			Vertices[0].UV = FVector2f(0.f, 0.f);
-
-			Vertices[1].Position = FVector4f(1.f, 0.f, 0.f, 1.f);
-			Vertices[1].UV = FVector2f(1.f, 0.f);
-
-			Vertices[2].Position = FVector4f(0.f, 1.f, 0.f, 1.f);
-			Vertices[2].UV = FVector2f(0.f, 1.f);
-
-			Vertices[3].Position = FVector4f(1.f, 1.f, 0.f, 1.f);
-			Vertices[3].UV = FVector2f(1.f, 1.f);
-
-			FRHIResourceCreateInfo CreateInfoVB(TEXT("FPostProcessSceneViewExtension"), &Vertices);
-			VertexBufferRHI = RHICreateVertexBuffer(Vertices.GetResourceDataSize(), BUF_Static, CreateInfoVB);
-		}
-
-		if (IndexBufferRHI == nullptr || !IndexBufferRHI.IsValid())
-		{
-			// Setup index buffer
-			const uint16 Indices[] = { 0, 1, 2, 2, 1, 3 };
-
-			TResourceArray<uint16, INDEXBUFFER_ALIGNMENT> IndexBuffer;
-			const uint32 NumIndices = UE_ARRAY_COUNT(Indices);
-			IndexBuffer.AddUninitialized(NumIndices);
-			FMemory::Memcpy(IndexBuffer.GetData(), Indices, NumIndices * sizeof(uint16));
-
-			FRHIResourceCreateInfo CreateInfoIB(TEXT("FPostProcessSceneViewExtension"), &IndexBuffer);
-			IndexBufferRHI = RHICreateIndexBuffer(sizeof(uint16), IndexBuffer.GetResourceDataSize(), BUF_Static, CreateInfoIB);
-		}
-
-		PostProcessMaterial = PostProcessMaterialGetter.Get();
-		const bool bIsValidMaterial = PostProcessMaterial && PostProcessMaterial->GetMaterial();
-		
-		const bool bHasCorrectDomain = bIsValidMaterial && PostProcessMaterial->GetMaterial()->MaterialDomain == EMaterialDomain::MD_PostProcess;
-		UE_CLOG(!bHasCorrectDomain, LogVPUtilities, Error, TEXT("Material Domain must be PostProcess!"));
-		
-		const bool bOutputsAlpha = bIsValidMaterial && PostProcessMaterial->GetMaterial()->BlendableOutputAlpha;
-		UE_CLOG(!bOutputsAlpha, LogVPUtilities, Error, TEXT("Material must have \"Output Alpha\" checked!"));
-		
-		const bool bCanUseMaterial = bHasCorrectDomain && bOutputsAlpha;
-		if (!ensure(bCanUseMaterial))
-		{
-			PostProcessMaterial = nullptr;
-		}
-	}
+	{}
 	
 	void FPostProcessSceneViewExtension::PreRenderViewFamily_RenderThread(FRDGBuilder& GraphBuilder, FSceneViewFamily& InViewFamily)
 	{}
 
 	void FPostProcessSceneViewExtension::PostRenderViewFamily_RenderThread(FRDGBuilder& GraphBuilder,FSceneViewFamily& InViewFamily)
 	{
-		if (!PostProcessMaterial || !PostProcessMaterial->IsValidLowLevel() || !VertexBufferRHI.IsValid() || !IndexBufferRHI.IsValid())
+		if (!WidgetRenderTarget.IsValid())
 		{
 			return;
 		}
@@ -194,71 +55,67 @@ namespace UE::VirtualProductionUtilities::Private
 
 	bool FPostProcessSceneViewExtension::IsActiveThisFrame_Internal(const FSceneViewExtensionContext& Context) const
 	{
-		return ensure(PostProcessMaterialGetter.IsBound()) && PostProcessMaterialGetter.Get() != nullptr;
+		return WidgetRenderTarget.IsValid();
 	}
 
-	BEGIN_SHADER_PARAMETER_STRUCT(FVPFullScreenPostProcessPassParameters, )
-		SHADER_PARAMETER_RDG_UNIFORM_BUFFER(FSceneTextureUniformParameters, SceneTextures)
-		RENDER_TARGET_BINDING_SLOTS()
-	END_SHADER_PARAMETER_STRUCT()
+	/**
+	 * This shaders overlays WidgetTexture of SceneTexture by blending it like so
+	 *	color = WidgetTexture.A * WidgetTexture.RGB + (1 - WidgetTexture.A) * SceneTexture.RGB
+	 */
+	class FDrawTextureInShaderPS : public FGlobalShader
+	{
+	public:
+		DECLARE_GLOBAL_SHADER(FDrawTextureInShaderPS);
+		SHADER_USE_PARAMETER_STRUCT(FDrawTextureInShaderPS, FGlobalShader);
+
+		BEGIN_SHADER_PARAMETER_STRUCT(FParameters, )
+			SHADER_PARAMETER_RDG_TEXTURE(Texture2D, WidgetTexture) 
+			SHADER_PARAMETER_SAMPLER(SamplerState, WidgetSampler)
+			SHADER_PARAMETER_RDG_TEXTURE(Texture2D, SceneTexture) 
+			SHADER_PARAMETER_SAMPLER(SamplerState, SceneSampler)
+			RENDER_TARGET_BINDING_SLOTS()
+		END_SHADER_PARAMETER_STRUCT()
+
+		static void ModifyCompilationEnvironment(const FPermutationParameters& Parameters, FShaderCompilerEnvironment& OutEnvironment)
+		{}
+	};
+	IMPLEMENT_GLOBAL_SHADER(FDrawTextureInShaderPS, "/Plugin/VirtualProductionUtilities/Private/VPFullScreenWidgetOverlay.usf", "OverlayWidgetPS", SF_Pixel);
 	
 	void FPostProcessSceneViewExtension::RenderMaterial_RenderThread(FRDGBuilder& GraphBuilder, const FSceneView& InView, FRDGTextureRef ViewFamilyTexture)
 	{
-		if (!VertexBufferRHI.IsValid() || !IndexBufferRHI.IsValid())
+		// Can be invalidated after exiting PIE
+		if (!WidgetRenderTarget.IsValid() || !WidgetRenderTarget->GetRenderTargetResource() || !WidgetRenderTarget->GetRenderTargetResource()->GetTexture2DRHI())
 		{
 			return;
 		}
+		
+		const FTexture2DRHIRef WidgetRenderTarget_RHI = WidgetRenderTarget->GetRenderTargetResource()->GetTexture2DRHI();
+		const FRDGTextureRef WidgetRenderTarget_RDG = GraphBuilder.RegisterExternalTexture(CreateRenderTarget(WidgetRenderTarget_RHI, TEXT("WidgetRenderTarget")));
 
-		auto* PassParameters = GraphBuilder.AllocParameters<FVPFullScreenPostProcessPassParameters>();
-		PassParameters->RenderTargets[0] = FRenderTargetBinding(ViewFamilyTexture, ERenderTargetLoadAction::ELoad);
-		check(InView.bIsViewInfo);
-		PassParameters->SceneTextures = CreateSceneTextureUniformBuffer(GraphBuilder, InView, ESceneTextureSetupMode::None);
-
-		GraphBuilder.AddPass(
+		const FGlobalShaderMap* GlobalShaderMap = GetGlobalShaderMap(GMaxRHIFeatureLevel);
+		const TShaderMapRef<FScreenPassVS> VertexShader(GlobalShaderMap);
+		const TShaderMapRef<FDrawTextureInShaderPS> PixelShader(GlobalShaderMap);
+		
+		FDrawTextureInShaderPS::FParameters* Parameters = GraphBuilder.AllocParameters<FDrawTextureInShaderPS::FParameters>();
+		Parameters->WidgetTexture = WidgetRenderTarget_RDG;
+		Parameters->WidgetSampler = TStaticSamplerState<SF_Point>::GetRHI();
+		Parameters->SceneTexture = ViewFamilyTexture;
+		Parameters->SceneSampler = TStaticSamplerState<SF_Point>::GetRHI();
+		Parameters->RenderTargets[0] = FRenderTargetBinding{ ViewFamilyTexture, ERenderTargetLoadAction::ELoad };
+
+		const FScreenPassTextureViewport InputViewport(WidgetRenderTarget_RDG);
+		const FScreenPassTextureViewport OutputViewport(ViewFamilyTexture);
+		// Note that we reference WidgetRenderTarget here, which means RDG will synchronize access to it.
+		// That means that the DrawWindow operation (see FVPFullScreenUserWidget_PostProcessBase) will finish writing into WidgetRenderTarget before we access it with this pass.
+		AddDrawScreenPass(
+			GraphBuilder,
 			RDG_EVENT_NAME("VPFullScreenPostProcessOverlay"),
-			PassParameters,
-			ERDGPassFlags::Raster,
-			[this, &InView](FRHICommandList& RHICmdList)
-		{
-			const auto FeatureLevel = InView.GetFeatureLevel();
-
-			const FMaterialRenderProxy* MaterialProxy = PostProcessMaterial->GetRenderProxy();
-			const FMaterial& CameraMaterial = MaterialProxy->GetMaterialWithFallback(FeatureLevel, MaterialProxy);
-			const FMaterialShaderMap* const MaterialShaderMap = CameraMaterial.GetRenderingThreadShaderMap();
-
-			FGraphicsPipelineStateInitializer GraphicsPSOInit;
-			RHICmdList.ApplyCachedRenderTargets(GraphicsPSOInit);
-
-			// To overlay the widget, we'll do a lerp: Final color = src color * src alpha + dest color * (1-src alpha).
-			// Remember: Src = PS output = widget; Dst = current render target contents = scene color
-			// Note that "Output Alpha" must be checked in the post process material!!! This makes MATERIAL_OUTPUT_OPACITY_AS_ALPHA = 1 in the PS so the widget opacity is passed as alpha value.
-			GraphicsPSOInit.BlendState = TStaticBlendState<CW_RGB, BO_Add, BF_SourceAlpha, BF_InverseSourceAlpha>::GetRHI();
-			GraphicsPSOInit.RasterizerState = TStaticRasterizerState<>::GetRHI();
-			GraphicsPSOInit.DepthStencilState = TStaticDepthStencilState<false, CF_Always>::GetRHI();	
-				
-			GraphicsPSOInit.PrimitiveType = PT_TriangleList;
-			GraphicsPSOInit.BoundShaderState.VertexDeclarationRHI = GFilterVertexDeclaration.VertexDeclarationRHI;
-
-			TShaderRef<FVPFullScreenPostProcessVS> VertexShader = MaterialShaderMap->GetShader<FVPFullScreenPostProcessVS>();
-			TShaderRef<FVPFullScreenPostProcessPS> PixelShader = MaterialShaderMap->GetShader<FVPFullScreenPostProcessPS>();
-			GraphicsPSOInit.BoundShaderState.VertexShaderRHI = VertexShader.GetVertexShader();
-			GraphicsPSOInit.BoundShaderState.PixelShaderRHI = PixelShader.GetPixelShader();
-
-			SetGraphicsPipelineState(RHICmdList, GraphicsPSOInit, 0);
-
-			SetShaderParametersLegacyVS(RHICmdList, VertexShader, InView);
-			SetShaderParametersLegacyPS(RHICmdList, PixelShader, InView, MaterialProxy);
-
-			RHICmdList.SetStreamSource(0, VertexBufferRHI, 0);
-			RHICmdList.DrawIndexedPrimitive(
-				IndexBufferRHI,
-				/*BaseVertexIndex=*/ 0,
-				/*MinIndex=*/ 0,
-				/*NumVertices=*/ 4,
-				/*StartIndex=*/ 0,
-				/*NumPrimitives=*/ 2,
-				/*NumInstances=*/ 1
+			InView,
+			OutputViewport,
+			InputViewport,
+			VertexShader,
+			PixelShader,
+			Parameters
 			);
-		});
 	}
-}
\ No newline at end of file
+}
diff --git a/Engine/Plugins/Experimental/VirtualProductionUtilities/Source/VPUtilities/Private/Widgets/Fullscreen/PostProcessSceneViewExtension.h b/Engine/Plugins/Experimental/VirtualProductionUtilities/Source/VPUtilities/Private/Widgets/Fullscreen/PostProcessSceneViewExtension.h
--- a/Engine/Plugins/Experimental/VirtualProductionUtilities/Source/VPUtilities/Private/Widgets/Fullscreen/PostProcessSceneViewExtension.h
+++ b/Engine/Plugins/Experimental/VirtualProductionUtilities/Source/VPUtilities/Private/Widgets/Fullscreen/PostProcessSceneViewExtension.h
@@ -6,6 +6,7 @@
 #include "UObject/WeakObjectPtrFwd.h"
 
 class UMaterialInterface;
+class UTextureRenderTarget2D;
 
 namespace UE::VirtualProductionUtilities::Private
 {
@@ -24,7 +25,7 @@ namespace UE::VirtualProductionUtilities::Private
 		using Super = FSceneViewExtensionBase;
 	public:
 		
-		FPostProcessSceneViewExtension(const FAutoRegister& AutoRegister, TAttribute<UMaterialInterface*> PostProcessMaterialGetter);
+		FPostProcessSceneViewExtension(const FAutoRegister& AutoRegister, UTextureRenderTarget2D& WidgetRenderTarget);
 
 		//~ISceneViewExtension interface
 		virtual void SetupViewFamily(FSceneViewFamily& InViewFamily) override;
@@ -37,16 +38,9 @@ namespace UE::VirtualProductionUtilities::Private
 		//~ISceneViewExtension interface
 
 	private:
-
-		/** Gets the post process material to render. */
-		TAttribute<UMaterialInterface*> PostProcessMaterialGetter;
-		/** Set on RenderThread. It is the responsibility this extension's user for preventing GC. */
-		UMaterialInterface* PostProcessMaterial = nullptr;
-
-		/** Index buffer for drawing the quad. Initialized the first time PreRenderView_RenderThread runs. */
-		FBufferRHIRef IndexBufferRHI;
-		/** Vertex buffer for drawing the quad. Initialized the first time PreRenderView_RenderThread runs. */
-		FBufferRHIRef VertexBufferRHI;
+		
+		/** Contains the widget that is supposed to be overlaid. */
+		TWeakObjectPtr<UTextureRenderTarget2D> WidgetRenderTarget;
 		
 		void RenderMaterial_RenderThread(FRDGBuilder& GraphBuilder, const FSceneView& InView, FRDGTextureRef ViewFamilyTexture);
 	};
diff --git a/Engine/Plugins/Experimental/VirtualProductionUtilities/Source/VPUtilities/Private/Widgets/Fullscreen/VPFullScreenUserWidget_PostProcess.cpp b/Engine/Plugins/Experimental/VirtualProductionUtilities/Source/VPUtilities/Private/Widgets/Fullscreen/VPFullScreenUserWidget_PostProcess.cpp
--- a/Engine/Plugins/Experimental/VirtualProductionUtilities/Source/VPUtilities/Private/Widgets/Fullscreen/VPFullScreenUserWidget_PostProcess.cpp
+++ b/Engine/Plugins/Experimental/VirtualProductionUtilities/Source/VPUtilities/Private/Widgets/Fullscreen/VPFullScreenUserWidget_PostProcess.cpp
@@ -2,11 +2,10 @@
 
 #include "Widgets/Fullscreen/VPFullScreenUserWidget_PostProcess.h"
 
-#include "VPUtilitiesModule.h"
 #include "Widgets/VPFullScreenUserWidget.h"
 
 #include "Components/PostProcessComponent.h"
-#include "Components/WidgetComponent.h"
+#include "Engine/TextureRenderTarget2D.h"
 #include "HAL/PlatformApplicationMisc.h"
 #include "GameFramework/WorldSettings.h"
 #include "Materials/MaterialInstanceDynamic.h"
@@ -17,6 +16,13 @@
 #include "SLevelViewport.h"
 #endif
 
+namespace UE::VPUtilities::Private
+{
+	const FName NAME_SlateUI = "SlateUI";
+	const FName NAME_TintColorAndOpacity = "TintColorAndOpacity";
+	const FName NAME_OpacityFromTexture = "OpacityFromTexture";
+}
+
 FVPFullScreenUserWidget_PostProcess::FVPFullScreenUserWidget_PostProcess()
 	: PostProcessComponent(nullptr)
 {}
@@ -25,11 +31,11 @@ void FVPFullScreenUserWidget_PostProcess::SetCustomPostProcessSettingsSource(TWe
 {
 	CustomPostProcessSettingsSource = InCustomPostProcessSettingsSource;
 	
-	const bool bIsRunning = GetPostProcessMaterialInstance() != nullptr;
+	const bool bIsRunning = PostProcessMaterialInstance != nullptr;
 	if (bIsRunning)
 	{
 		// Save us from creating another struct member: PostProcessMaterialInstance is always created with UWorld as outer.
-		UWorld* World = CastChecked<UWorld>(GetPostProcessMaterialInstance()->GetOuter());
+		UWorld* World = CastChecked<UWorld>(PostProcessMaterialInstance->GetOuter());
 		InitPostProcessComponent(World);
 	}
 }
@@ -85,7 +91,7 @@ bool FVPFullScreenUserWidget_PostProcess::InitPostProcessComponent(UWorld* World
 
 bool FVPFullScreenUserWidget_PostProcess::UpdateTargetPostProcessSettingsWithMaterial()
 {
-	UMaterialInstanceDynamic* MaterialInstance = GetPostProcessMaterialInstance();
+	UMaterialInstanceDynamic* MaterialInstance = PostProcessMaterialInstance;
 	if (FPostProcessSettings* const PostProcessSettings = GetPostProcessSettings()
 		; PostProcessSettings && ensure(MaterialInstance))
 	{
@@ -113,7 +119,7 @@ bool FVPFullScreenUserWidget_PostProcess::UpdateTargetPostProcessSettingsWithMat
 
 void FVPFullScreenUserWidget_PostProcess::ReleasePostProcessComponent()
 {
-	const bool bIsRunning = GetPostProcessMaterialInstance() != nullptr;
+	const bool bIsRunning = PostProcessMaterialInstance != nullptr;
 	if (!bIsRunning)
 	{
 		return;
@@ -123,7 +129,7 @@ void FVPFullScreenUserWidget_PostProcess::ReleasePostProcessComponent()
 	if (FPostProcessSettings* Settings = GetPostProcessSettings()
 		; bNeedsToResetExternalSettings && Settings)
 	{
-		const int32 Index = Settings->WeightedBlendables.Array.IndexOfByPredicate([this](const FWeightedBlendable& Blendable){ return Blendable.Object == GetPostProcessMaterialInstance(); });
+		const int32 Index = Settings->WeightedBlendables.Array.IndexOfByPredicate([this](const FWeightedBlendable& Blendable){ return Blendable.Object == PostProcessMaterialInstance; });
 		if (Index != INDEX_NONE)
 		{
 			Settings->WeightedBlendables.Array.RemoveAt(Index);
@@ -138,6 +144,22 @@ void FVPFullScreenUserWidget_PostProcess::ReleasePostProcessComponent()
 	PostProcessComponent = nullptr;
 }
 
+bool FVPFullScreenUserWidget_PostProcess::OnRenderTargetInited()
+{
+	// Outer needs to be transient package: otherwise we cause a world memory leak using "Save Current Level As" due to reference not getting replaced correctly
+	PostProcessMaterialInstance = UMaterialInstanceDynamic::Create(PostProcessMaterial, GetTransientPackage());
+	PostProcessMaterialInstance->SetFlags(RF_Transient);
+	if (ensure(PostProcessMaterialInstance))
+	{
+		using namespace UE::VPUtilities::Private;
+		PostProcessMaterialInstance->SetTextureParameterValue(NAME_SlateUI, WidgetRenderTarget);
+		PostProcessMaterialInstance->SetVectorParameterValue(NAME_TintColorAndOpacity, PostProcessTintColorAndOpacity);
+		PostProcessMaterialInstance->SetScalarParameterValue(NAME_OpacityFromTexture, PostProcessOpacityFromTexture);
+		return true;
+	}
+	return false;
+}
+
 FPostProcessSettings* FVPFullScreenUserWidget_PostProcess::GetPostProcessSettings() const
 {
 	if (PostProcessComponent)
@@ -161,4 +183,4 @@ FPostProcessSettings* FVPFullScreenUserWidget_PostProcess::GetPostProcessSetting
 	}
 
 	return nullptr;
-}
\ No newline at end of file
+}
diff --git a/Engine/Plugins/Experimental/VirtualProductionUtilities/Source/VPUtilities/Private/Widgets/Fullscreen/VPFullScreenUserWidget_PostProcessBase.cpp b/Engine/Plugins/Experimental/VirtualProductionUtilities/Source/VPUtilities/Private/Widgets/Fullscreen/VPFullScreenUserWidget_PostProcessBase.cpp
--- a/Engine/Plugins/Experimental/VirtualProductionUtilities/Source/VPUtilities/Private/Widgets/Fullscreen/VPFullScreenUserWidget_PostProcessBase.cpp
+++ b/Engine/Plugins/Experimental/VirtualProductionUtilities/Source/VPUtilities/Private/Widgets/Fullscreen/VPFullScreenUserWidget_PostProcessBase.cpp
@@ -15,9 +15,7 @@
 #include "Framework/Application/SlateApplication.h"
 #include "Input/HittestGrid.h"
 #include "HAL/PlatformApplicationMisc.h"
-#include "GameFramework/WorldSettings.h"
 #include "Layout/Visibility.h"
-#include "Materials/MaterialInstanceDynamic.h"
 #include "RHI.h"
 #include "Slate/SceneViewport.h"
 #include "Slate/WidgetRenderer.h"
@@ -28,16 +26,11 @@
 
 #if WITH_EDITOR
 #include "LevelEditor.h"
-#include "Modules/ModuleManager.h"
 #include "SLevelViewport.h"
 #endif
 
 namespace UE::VPUtilities::Private
 {
-	const FName NAME_SlateUI = "SlateUI";
-	const FName NAME_TintColorAndOpacity = "TintColorAndOpacity";
-	const FName NAME_OpacityFromTexture = "OpacityFromTexture";
-
 	EVisibility ConvertWindowVisibilityToVisibility(EWindowVisibility visibility)
 	{
 		switch (visibility)
@@ -131,7 +124,6 @@ FVPFullScreenUserWidget_PostProcessBase::FVPFullScreenUserWidget_PostProcessBase
 	, RenderTargetBackgroundColor(FLinearColor(0.0f, 0.0f, 0.0f, 1.0f))
 	, RenderTargetBlendMode(EWidgetBlendMode::Masked)
 	, WidgetRenderTarget(nullptr)
-	, PostProcessMaterialInstance(nullptr)
 	, WidgetRenderer(nullptr)
 	, CurrentWidgetDrawSize(FIntPoint::ZeroValue)
 {}
@@ -146,11 +138,6 @@ TSharedPtr<SVirtualWindow> FVPFullScreenUserWidget_PostProcessBase::GetSlateWind
 	return SlateWindow;
 }
 
-TObjectPtr<UMaterialInstanceDynamic> FVPFullScreenUserWidget_PostProcessBase::GetPostProcessMaterialInstance() const
-{
-	return PostProcessMaterialInstance;
-}
-
 bool FVPFullScreenUserWidget_PostProcessBase::CreateRenderer(UWorld* World, UUserWidget* Widget, TAttribute<float> InDPIScale)
 {
 	ReleaseRenderer();
@@ -207,10 +194,10 @@ bool FVPFullScreenUserWidget_PostProcessBase::CreateRenderer(UWorld* World, UUse
 		WidgetRenderTarget = NewObject<UTextureRenderTarget2D>(GetTransientPackage(), NAME_None, RF_Transient);
 		WidgetRenderTarget->ClearColor = ActualBackgroundColor;
 
-		InitPostProcessMaterial();
+		return WidgetRenderer && WidgetRenderTarget && OnRenderTargetInited();
 	}
 
-	return WidgetRenderer && WidgetRenderTarget && PostProcessMaterialInstance;
+	return false;
 }
 
 void FVPFullScreenUserWidget_PostProcessBase::ReleaseRenderer()
@@ -225,7 +212,6 @@ void FVPFullScreenUserWidget_PostProcessBase::ReleaseRenderer()
 	SlateWindow.Reset();
 	WidgetRenderTarget = nullptr;
 	CurrentWidgetDrawSize = FIntPoint::ZeroValue;
-	PostProcessMaterialInstance = nullptr;
 }
 
 void FVPFullScreenUserWidget_PostProcessBase::TickRenderer(UWorld* World, float DeltaSeconds)
@@ -269,20 +255,6 @@ void FVPFullScreenUserWidget_PostProcessBase::TickRenderer(UWorld* World, float
 	}
 }
 
-void FVPFullScreenUserWidget_PostProcessBase::InitPostProcessMaterial()
-{
-	// Outer needs to be transient package: otherwise we cause a world memory leak using "Save Current Level As" due to reference not getting replaced correctly
-	PostProcessMaterialInstance = UMaterialInstanceDynamic::Create(PostProcessMaterial, GetTransientPackage());
-	PostProcessMaterialInstance->SetFlags(RF_Transient);
-	if (ensure(PostProcessMaterialInstance))
-	{
-		using namespace UE::VPUtilities::Private;
-		PostProcessMaterialInstance->SetTextureParameterValue(NAME_SlateUI, WidgetRenderTarget);
-		PostProcessMaterialInstance->SetVectorParameterValue(NAME_TintColorAndOpacity, PostProcessTintColorAndOpacity);
-		PostProcessMaterialInstance->SetScalarParameterValue(NAME_OpacityFromTexture, PostProcessOpacityFromTexture);
-	}
-}
-
 FIntPoint FVPFullScreenUserWidget_PostProcessBase::CalculateWidgetDrawSize(UWorld* World)
 {
 	if (bUseWidgetDrawSize)
diff --git a/Engine/Plugins/Experimental/VirtualProductionUtilities/Source/VPUtilities/Private/Widgets/Fullscreen/VPFullScreenUserWidget_PostProcessWithSVE.cpp b/Engine/Plugins/Experimental/VirtualProductionUtilities/Source/VPUtilities/Private/Widgets/Fullscreen/VPFullScreenUserWidget_PostProcessWithSVE.cpp
--- a/Engine/Plugins/Experimental/VirtualProductionUtilities/Source/VPUtilities/Private/Widgets/Fullscreen/VPFullScreenUserWidget_PostProcessWithSVE.cpp
+++ b/Engine/Plugins/Experimental/VirtualProductionUtilities/Source/VPUtilities/Private/Widgets/Fullscreen/VPFullScreenUserWidget_PostProcessWithSVE.cpp
@@ -3,6 +3,7 @@
 #include "Widgets/Fullscreen/VPFullScreenUserWidget_PostProcessWithSVE.h"
 
 #include "Widgets/Fullscreen/PostProcessSceneViewExtension.h"
+#include "Widgets/VPFullScreenUserWidget.h"
 
 #include "Materials/MaterialInterface.h"
 #include "Materials/MaterialInstanceDynamic.h"
@@ -11,10 +12,10 @@
 bool FVPFullScreenUserWidget_PostProcessWithSVE::Display(UWorld* World, UUserWidget* Widget, TAttribute<float> InDPIScale)
 {
 	bool bOk = CreateRenderer(World, Widget, MoveTemp(InDPIScale));
-	if (bOk && ensureMsgf(GetPostProcessMaterialInstance(), TEXT("CreateRenderer returned true even though it failed.")))
+	if (bOk && ensureMsgf(WidgetRenderTarget, TEXT("CreateRenderer returned true even though it failed.")))
 	{
 		SceneViewExtension = FSceneViewExtensions::NewExtension<UE::VirtualProductionUtilities::Private::FPostProcessSceneViewExtension>(
-			TAttribute<UMaterialInterface*>::CreateRaw(this, &FVPFullScreenUserWidget_PostProcessWithSVE::GetPostProcessMaterial)
+			*WidgetRenderTarget
 			);
 		SceneViewExtension->IsActiveThisFrameFunctions = MoveTemp(IsActiveFunctorsToRegister);
 	}
@@ -43,8 +44,3 @@ void FVPFullScreenUserWidget_PostProcessWithSVE::RegisterIsActiveFunctor(FSceneV
 		IsActiveFunctorsToRegister.Emplace(MoveTemp(IsActiveFunctor));
 	}
 }
-
-UMaterialInterface* FVPFullScreenUserWidget_PostProcessWithSVE::GetPostProcessMaterial() const
-{
-	return GetPostProcessMaterialInstance().Get();
-}
diff --git a/Engine/Plugins/Experimental/VirtualProductionUtilities/Source/VPUtilities/Public/Widgets/Fullscreen/VPFullScreenUserWidget_PostProcess.h b/Engine/Plugins/Experimental/VirtualProductionUtilities/Source/VPUtilities/Public/Widgets/Fullscreen/VPFullScreenUserWidget_PostProcess.h
--- a/Engine/Plugins/Experimental/VirtualProductionUtilities/Source/VPUtilities/Public/Widgets/Fullscreen/VPFullScreenUserWidget_PostProcess.h
+++ b/Engine/Plugins/Experimental/VirtualProductionUtilities/Source/VPUtilities/Public/Widgets/Fullscreen/VPFullScreenUserWidget_PostProcess.h
@@ -33,6 +33,10 @@ private:
 	/** Post process component used to add the material to the post process chain. */
 	UPROPERTY(Transient)
 	TObjectPtr<UPostProcessComponent> PostProcessComponent;
+	
+	/** The dynamic instance of the material that the render target is attached to. */
+	UPROPERTY(Transient)
+	TObjectPtr<UMaterialInstanceDynamic> PostProcessMaterialInstance;
 
 	/** Only render to the UTextureRenderTarget2D - do not output to the final viewport. */
 	bool bRenderToTextureOnly;
@@ -46,6 +50,8 @@ private:
 	bool InitPostProcessComponent(UWorld* World);
 	bool UpdateTargetPostProcessSettingsWithMaterial();
 	void ReleasePostProcessComponent();
+	
+	virtual bool OnRenderTargetInited() override;
 
 	FPostProcessSettings* GetPostProcessSettings() const;
 };
diff --git a/Engine/Plugins/Experimental/VirtualProductionUtilities/Source/VPUtilities/Public/Widgets/Fullscreen/VPFullScreenUserWidget_PostProcessBase.h b/Engine/Plugins/Experimental/VirtualProductionUtilities/Source/VPUtilities/Public/Widgets/Fullscreen/VPFullScreenUserWidget_PostProcessBase.h
--- a/Engine/Plugins/Experimental/VirtualProductionUtilities/Source/VPUtilities/Public/Widgets/Fullscreen/VPFullScreenUserWidget_PostProcessBase.h
+++ b/Engine/Plugins/Experimental/VirtualProductionUtilities/Source/VPUtilities/Public/Widgets/Fullscreen/VPFullScreenUserWidget_PostProcessBase.h
@@ -97,20 +97,15 @@ struct FVPFullScreenUserWidget_PostProcessBase
 	virtual void Hide(UWorld* World);
 	
 	TSharedPtr<SVirtualWindow> VPUTILITIES_API GetSlateWindow() const;
-	TObjectPtr<UMaterialInstanceDynamic> GetPostProcessMaterialInstance() const;
 
 protected:
 
 	bool CreateRenderer(UWorld* World, UUserWidget* Widget, TAttribute<float> InDPIScale);
-	void ReleaseRenderer();
+	virtual void ReleaseRenderer();
 	void TickRenderer(UWorld* World, float DeltaSeconds);
 	
 private:
 	
-	/** The dynamic instance of the material that the render target is attached to. */
-	UPROPERTY(Transient)
-	TObjectPtr<UMaterialInstanceDynamic> PostProcessMaterialInstance;
-	
 	/** Helper class for drawing widgets to a render target. */
 	FWidgetRenderer* WidgetRenderer;
 	
@@ -126,7 +121,7 @@ private:
 	TSharedPtr<UE::VPUtilities::Private::FVPWidgetPostProcessHitTester> CustomHitTestPath;
 
 	/** Creates the post process material and sets up its parameters. */
-	void InitPostProcessMaterial();
+	virtual bool OnRenderTargetInited() { return true; };
 	
 	/** Determines widget size depending on the viewport type (PIE / Game) */
 	FIntPoint CalculateWidgetDrawSize(UWorld* World);
diff --git a/Engine/Plugins/Experimental/VirtualProductionUtilities/Source/VPUtilities/Public/Widgets/Fullscreen/VPFullScreenUserWidget_PostProcessWithSVE.h b/Engine/Plugins/Experimental/VirtualProductionUtilities/Source/VPUtilities/Public/Widgets/Fullscreen/VPFullScreenUserWidget_PostProcessWithSVE.h
--- a/Engine/Plugins/Experimental/VirtualProductionUtilities/Source/VPUtilities/Public/Widgets/Fullscreen/VPFullScreenUserWidget_PostProcessWithSVE.h
+++ b/Engine/Plugins/Experimental/VirtualProductionUtilities/Source/VPUtilities/Public/Widgets/Fullscreen/VPFullScreenUserWidget_PostProcessWithSVE.h
@@ -8,6 +8,7 @@
 
 class FSceneViewExtensionBase;
 class UMaterialInterface;
+class UVPFullScreenUserWidget;
 
 /**
  * Renders widget in post process phase by using Scene View Extensions (SVE).
@@ -37,6 +38,4 @@ private:
 
 	/** Functors registered before we were displayed. */
 	TArray<FSceneViewExtensionIsActiveFunctor> IsActiveFunctorsToRegister;
-
-	UMaterialInterface* GetPostProcessMaterial() const;
 };
diff --git a/Engine/Plugins/Experimental/VirtualProductionUtilities/Source/VPUtilities/VPUtilities.Build.cs b/Engine/Plugins/Experimental/VirtualProductionUtilities/Source/VPUtilities/VPUtilities.Build.cs
--- a/Engine/Plugins/Experimental/VirtualProductionUtilities/Source/VPUtilities/VPUtilities.Build.cs
+++ b/Engine/Plugins/Experimental/VirtualProductionUtilities/Source/VPUtilities/VPUtilities.Build.cs
@@ -29,6 +29,7 @@ public class VPUtilities : ModuleRules
 			new string[]
 			{
 				"ApplicationCore",
+				"Projects",
 				"Renderer",
 				"RenderCore",
 				"RHI",
