Fix CIS related to BaseVirtualProductionUtilities

diff --git a/Engine/Plugins/Experimental/VirtualProductionUtilities/Config/BaseVirtualProductionUtilities.ini b/Engine/Plugins/Experimental/VirtualProductionUtilities/Config/BaseVirtualProductionUtilities.ini
--- a/Engine/Plugins/Experimental/VirtualProductionUtilities/Config/BaseVirtualProductionUtilities.ini
+++ b/Engine/Plugins/Experimental/VirtualProductionUtilities/Config/BaseVirtualProductionUtilities.ini
@@ -13,7 +13,7 @@
 +PropertyRedirects=(OldName="VPFullScreenUserWidget_PostProcess.WidgetDrawSize",NewName="VPFullScreenUserWidget_PostProcessBase.WidgetDrawSize")
 +PropertyRedirects=(OldName="VPFullScreenUserWidget_PostProcess.bWindowFocusable",NewName="VPFullScreenUserWidget_PostProcessBase.bWindowFocusable")
 +PropertyRedirects=(OldName="VPFullScreenUserWidget_PostProcess.WindowVisibility",NewName="VPFullScreenUserWidget_PostProcessBase.WindowVisibility")
-+PropertyRedirects=(OldName="VPFullScreenUserWidget_PostProcess.bReceiveHardwareInput",NewName=FVPFullScreenUserWidget_PostProcessBase.bReceiveHardwareInput")
++PropertyRedirects=(OldName="VPFullScreenUserWidget_PostProcess.bReceiveHardwareInput",NewName="FVPFullScreenUserWidget_PostProcessBase.bReceiveHardwareInput")
 +PropertyRedirects=(OldName="VPFullScreenUserWidget_PostProcess.RenderTargetBackgroundColor",NewName="VPFullScreenUserWidget_PostProcessBase.RenderTargetBackgroundColor")
 +PropertyRedirects=(OldName="VPFullScreenUserWidget_PostProcess.RenderTargetBlendMode",NewName="FVPFullScreenUserWidget_PostProcessBase.RenderTargetBlendMode")
 +PropertyRedirects=(OldName="VPFullScreenUserWidget.WidgetRenderTarget",NewName="VPFullScreenUserWidget_PostProcessBase.WidgetRenderTarget")
