Delete WidetPostProcessMaterial_SceneViewExtensions again and this time remove code references to it

diff --git a/Engine/Plugins/Experimental/VirtualProductionUtilities/Source/VPUtilities/Private/Widgets/VPFullScreenUserWidget.cpp b/Engine/Plugins/Experimental/VirtualProductionUtilities/Source/VPUtilities/Private/Widgets/VPFullScreenUserWidget.cpp
--- a/Engine/Plugins/Experimental/VirtualProductionUtilities/Source/VPUtilities/Private/Widgets/VPFullScreenUserWidget.cpp
+++ b/Engine/Plugins/Experimental/VirtualProductionUtilities/Source/VPUtilities/Private/Widgets/VPFullScreenUserWidget.cpp
@@ -194,8 +194,6 @@ UVPFullScreenUserWidget::UVPFullScreenUserWidget(const FObjectInitializer& Objec
 {
 	static ConstructorHelpers::FObjectFinder<UMaterialInterface> PostProcessMaterial_Finder(TEXT("/VirtualProductionUtilities/Materials/WidgetPostProcessMaterial"));
 	PostProcessDisplayTypeWithBlendMaterial.PostProcessMaterial = PostProcessMaterial_Finder.Object;
-	static ConstructorHelpers::FObjectFinder<UMaterialInterface> PostProcessMaterial_SceneViewExtensions_Finder(TEXT("/VirtualProductionUtilities/Materials/WidgetPostProcessMaterial_SceneViewExtensions"));
-	PostProcessWithSceneViewExtensions.PostProcessMaterial = PostProcessMaterial_SceneViewExtensions_Finder.Object;
 }
 
 void UVPFullScreenUserWidget::BeginDestroy()
