Fix non-unity error

diff --git a/Engine/Plugins/Experimental/VirtualProductionUtilities/Source/VPUtilities/Private/Widgets/Fullscreen/VPFullScreenUserWidget_PostProcess.cpp b/Engine/Plugins/Experimental/VirtualProductionUtilities/Source/VPUtilities/Private/Widgets/Fullscreen/VPFullScreenUserWidget_PostProcess.cpp
--- a/Engine/Plugins/Experimental/VirtualProductionUtilities/Source/VPUtilities/Private/Widgets/Fullscreen/VPFullScreenUserWidget_PostProcess.cpp
+++ b/Engine/Plugins/Experimental/VirtualProductionUtilities/Source/VPUtilities/Private/Widgets/Fullscreen/VPFullScreenUserWidget_PostProcess.cpp
@@ -9,6 +9,7 @@
 #include "HAL/PlatformApplicationMisc.h"
 #include "GameFramework/WorldSettings.h"
 #include "Materials/MaterialInstanceDynamic.h"
+#include "UObject/Package.h"
 #include "Widgets/Layout/SConstraintCanvas.h"
 
 #if WITH_EDITOR
