VScout/VPUtilities - fix for laser location start point for bookmark spline meshes

diff --git a/Engine/Plugins/Experimental/VirtualProductionUtilities/Source/VPUtilities/Private/Libraries/VPBlueprintLibrary.cpp b/Engine/Plugins/Experimental/VirtualProductionUtilities/Source/VPUtilities/Private/Libraries/VPBlueprintLibrary.cpp
--- a/Engine/Plugins/Experimental/VirtualProductionUtilities/Source/VPUtilities/Private/Libraries/VPBlueprintLibrary.cpp
+++ b/Engine/Plugins/Experimental/VirtualProductionUtilities/Source/VPUtilities/Private/Libraries/VPBlueprintLibrary.cpp
@@ -215,6 +215,32 @@ FVector UVPBlueprintLibrary::GetVREditorLaserHoverLocation()
 	return FVector(0);
 }
 
+FVector UVPBlueprintLibrary::GetLeftInteractorLocation()
+{
+#if WITH_EDITOR
+	FString ErrorText(TEXT("VR interactor location will be invalid."));
+	if (const UViewportWorldInteraction* VPI = VPBlueprintLibrary::GetViewportWorldInteraction(ErrorText))
+	{
+		const TArray<UViewportInteractor*> Interactors = VPI->GetInteractors();
+
+		for (UViewportInteractor* Interactor : Interactors)
+		{
+			if (UVREditorInteractor* EdInteractor = Cast<UVREditorInteractor>(Interactor))
+			{
+				if (EdInteractor->GetControllerType() == EControllerType::Laser)
+				{
+					const FVector InteractorLocation = EdInteractor->GetInteractorData().Transform.GetLocation();
+					UE_LOG(LogVPUtilities, Verbose , TEXT("%s"), *InteractorLocation.ToString());
+
+					return InteractorLocation;
+				}
+			}
+		}
+	}
+#endif
+
+	return FVector(0);
+}
 
 bool UVPBlueprintLibrary::EditorUndo()
 {
@@ -294,7 +320,7 @@ void UVPBlueprintLibrary::VPBookmarkSplineMeshIndicatorSetStartAndEnd(USplineMes
 		
 	// @todo: Fix - GetVREditorLaserHoverLocation() does not return the correct hover location
 	// USplineMeshComponent::SetEndPosition expects local space
-	SplineMesh->SetEndPosition(SplineTransform.InverseTransformPosition(IsVREditorModeActive() ? GetVREditorLaserHoverLocation() : GetEditorViewportTransform().TransformPosition(FVector(80, 0, -20))));
+	SplineMesh->SetEndPosition(SplineTransform.InverseTransformPosition(IsVREditorModeActive() ? GetLeftInteractorLocation(): GetEditorViewportTransform().TransformPosition(FVector(80, 0, -20))));
 }
 
 void UVPBlueprintLibrary::VPBookmarkSplineMeshIndicatorDisable(USplineMeshComponent* SplineMesh)
diff --git a/Engine/Plugins/Experimental/VirtualProductionUtilities/Source/VPUtilities/Public/Libraries/VPBlueprintLibrary.h b/Engine/Plugins/Experimental/VirtualProductionUtilities/Source/VPUtilities/Public/Libraries/VPBlueprintLibrary.h
--- a/Engine/Plugins/Experimental/VirtualProductionUtilities/Source/VPUtilities/Public/Libraries/VPBlueprintLibrary.h
+++ b/Engine/Plugins/Experimental/VirtualProductionUtilities/Source/VPUtilities/Public/Libraries/VPBlueprintLibrary.h
@@ -68,6 +68,10 @@ public:
 	UFUNCTION(BlueprintCallable, Category = "Virtual Production")
 	static FVector GetVREditorLaserHoverLocation();
 
+	/** Get the location of the left interactor, in world space */
+	UFUNCTION(BlueprintCallable, Category = "Virtual Production")
+	static FVector GetLeftInteractorLocation();
+
 	/** Trigger an UnrealEd Undo */
 	UFUNCTION(BlueprintCallable, Category = "Virtual Production")
 	static bool EditorUndo();
