Gaffer tool: Gesture menu prototype

diff --git a/Engine/Plugins/Experimental/VirtualProductionUtilities/Source/VPUtilities/Private/Actors/VPFullScreenUserWidgetActor.cpp b/Engine/Plugins/Experimental/VirtualProductionUtilities/Source/VPUtilities/Private/Actors/VPFullScreenUserWidgetActor.cpp
--- a/Engine/Plugins/Experimental/VirtualProductionUtilities/Source/VPUtilities/Private/Actors/VPFullScreenUserWidgetActor.cpp
+++ b/Engine/Plugins/Experimental/VirtualProductionUtilities/Source/VPUtilities/Private/Actors/VPFullScreenUserWidgetActor.cpp
@@ -119,4 +119,13 @@ void AFullScreenUserWidgetActor::RequestGameDisplay()
 	}
 }
 
+UUserWidget* AFullScreenUserWidgetActor::GetUserWidget() const
+{
+	if (ScreenUserWidget)
+	{
+		return ScreenUserWidget->GetWidget();
+	}
+	return nullptr;
+}
+
 #undef LOCTEXT_NAMESPACE
diff --git a/Engine/Plugins/Experimental/VirtualProductionUtilities/Source/VPUtilities/Public/Actors/VPFullScreenUserWidgetActor.h b/Engine/Plugins/Experimental/VirtualProductionUtilities/Source/VPUtilities/Public/Actors/VPFullScreenUserWidgetActor.h
--- a/Engine/Plugins/Experimental/VirtualProductionUtilities/Source/VPUtilities/Public/Actors/VPFullScreenUserWidgetActor.h
+++ b/Engine/Plugins/Experimental/VirtualProductionUtilities/Source/VPUtilities/Public/Actors/VPFullScreenUserWidgetActor.h
@@ -35,6 +35,10 @@ public:
 #endif //WITH_EDITOR
 	//~ End AActor Interface
 
+	/* Get a pointer to the inner user widget  */
+	UFUNCTION(BlueprintCallable, Category = "User Interface")
+	UUserWidget* GetUserWidget() const;
+
 private:
 	void RequestEditorDisplay();
 	void RequestGameDisplay();
