Virtual Production utilities - Renamed display names of VPUtilities Virtual Scouting tools' project properties to "Legacy"

- Made the VPScoutingMode display name "Legacy Virtual Scouting"
- Marked all the project settings relating to the Legacy Virtual Scouting tools with a displayname  to match the VRMode dropdown on the main toolbar UI

diff --git a/Engine/Plugins/Experimental/VirtualProductionUtilities/Source/VPUtilitiesEditor/Public/VPUtilitiesEditorSettings.h b/Engine/Plugins/Experimental/VirtualProductionUtilities/Source/VPUtilitiesEditor/Public/VPUtilitiesEditorSettings.h
--- a/Engine/Plugins/Experimental/VirtualProductionUtilities/Source/VPUtilitiesEditor/Public/VPUtilitiesEditorSettings.h
+++ b/Engine/Plugins/Experimental/VirtualProductionUtilities/Source/VPUtilitiesEditor/Public/VPUtilitiesEditorSettings.h
@@ -22,35 +22,35 @@ public:
 	UVPUtilitiesEditorSettings();
 
 	/** The default user interface that we'll use for virtual scouting */
-	UPROPERTY(EditAnywhere, config, Category = "Virtual Production", meta = (DisplayName = "Virtual Scouting User Interface"))
+	UPROPERTY(EditAnywhere, config, Category = "Legacy Virtual Scouting", meta = (DisplayName = "Virtual Scouting User Interface"))
 	TSoftClassPtr<UEditorUtilityWidget> VirtualScoutingUI;
 	
 	/** Speed when flying in VR*/
-	UPROPERTY(EditAnywhere, config, Category = "Virtual Production", meta = (ClampMin = "0.0", ClampMax = "1.0", DisplayName = "Virtual Scouting Flight Speed"))
+	UPROPERTY(EditAnywhere, config, Category = "Legacy Virtual Scouting", meta = (ClampMin = "0.0", ClampMax = "1.0", DisplayName = "Virtual Scouting Flight Speed"))
 	float FlightSpeed = 0.5f;
 	
 	/** Speed when using grip nav in VR */
-	UPROPERTY(EditAnywhere, config, Category = "Virtual Production", meta = (ClampMin = "0.0", ClampMax = "1.0", DisplayName = "Virtual Scouting Grip Nav Speed"))
+	UPROPERTY(EditAnywhere, config, Category = "Legacy Virtual Scouting", meta = (ClampMin = "0.0", ClampMax = "1.0", DisplayName = "Virtual Scouting Grip Nav Speed"))
 	float GripNavSpeed = 0.25f;
 	
 	/** Whether to use the metric system or imperial for measurements */
-	UPROPERTY(EditAnywhere, config, Category = "Virtual Production", meta = (DisplayName = "Show Measurements In Metric Units"))
+	UPROPERTY(EditAnywhere, config, Category = "Legacy Virtual Scouting", meta = (DisplayName = "Show Measurements In Metric Units"))
 	bool bUseMetric = false;
 	
 	/** Whether to enable or disable the transform gizmo */
-	UPROPERTY(EditAnywhere, config, Category = "Virtual Production", meta = (DisplayName = "Enable Transform Gizmo In VR"))
+	UPROPERTY(EditAnywhere, config, Category = "Legacy Virtual Scouting", meta = (DisplayName = "Enable Transform Gizmo In VR"))
 	bool bUseTransformGizmo = false;
 	
 	/** If true, the user will use inertia damping to stop after grip nav. Otherwise the user will just stop immediately */
-	UPROPERTY(EditAnywhere, config, Category = "Virtual Production", meta = (DisplayName = "Use Grip Inertia Damping"))
+	UPROPERTY(EditAnywhere, config, Category = "Legacy Virtual Scouting", meta = (DisplayName = "Use Grip Inertia Damping"))
 	bool bUseGripInertiaDamping = true;
 	
 	/** Damping applied to inertia */
-	UPROPERTY(EditAnywhere, config, Category = "Virtual Production", meta = (ClampMin = "0.0", ClampMax = "1.0", DisplayName = "Inertia Damping"))
+	UPROPERTY(EditAnywhere, config, Category = "Legacy Virtual Scouting", meta = (ClampMin = "0.0", ClampMax = "1.0", DisplayName = "Inertia Damping"))
 	float InertiaDamping = 0.95f;
 
 	/** Whether the helper system on the controllers is enabled */
-	UPROPERTY(EditAnywhere, config, Category = "Virtual Production", meta = (DisplayName = "Helper System Enabled"))
+	UPROPERTY(EditAnywhere, config, Category = "Legacy Virtual Scouting", meta = (DisplayName = "Helper System Enabled "))
 	bool bIsHelperSystemEnabled = true;
 	
 	/** When enabled, an OSC server will automatically start on launch. */
diff --git a/Engine/Plugins/Experimental/VirtualScouting/Source/VirtualScoutingEditor/Private/VPScoutingMode.h b/Engine/Plugins/Experimental/VirtualScouting/Source/VirtualScoutingEditor/Private/VPScoutingMode.h
--- a/Engine/Plugins/Experimental/VirtualScouting/Source/VirtualScoutingEditor/Private/VPScoutingMode.h
+++ b/Engine/Plugins/Experimental/VirtualScouting/Source/VirtualScoutingEditor/Private/VPScoutingMode.h
@@ -6,7 +6,7 @@
 #include "VPScoutingMode.generated.h"
 
 
-UCLASS(meta=(DisplayName="Virtual Scouting (Legacy VPUtilities)"))
+UCLASS(meta=(DisplayName="Legacy Virtual Scouting"))
 class UVPScoutingMode : public UVREditorMode
 {
 	GENERATED_BODY()
