VPUtilities: Disabled additional settings to the passthrough PPV, for convenience: film grain, chromatic aberration & lumen screen traces.

#rnx

diff --git a/Engine/Plugins/Experimental/VirtualProductionUtilities/Source/VPUtilities/Public/Actors/VPPassthroughPostProcessVolume.h b/Engine/Plugins/Experimental/VirtualProductionUtilities/Source/VPUtilities/Public/Actors/VPPassthroughPostProcessVolume.h
--- a/Engine/Plugins/Experimental/VirtualProductionUtilities/Source/VPUtilities/Public/Actors/VPPassthroughPostProcessVolume.h
+++ b/Engine/Plugins/Experimental/VirtualProductionUtilities/Source/VPUtilities/Public/Actors/VPPassthroughPostProcessVolume.h
@@ -56,5 +56,17 @@ public:
 
 		Settings.bOverride_VignetteIntensity = 1;
 		Settings.VignetteIntensity = 0.f;
+
+		Settings.bOverride_FilmGrainIntensity = 1;
+		Settings.FilmGrainIntensity = 0.0f;
+
+		Settings.bOverride_SceneFringeIntensity = 1;
+		Settings.SceneFringeIntensity = 0.0f;
+
+		Settings.bOverride_LumenFinalGatherScreenTraces = 1;
+		Settings.LumenFinalGatherScreenTraces = 0;
+
+		Settings.bOverride_LumenReflectionsScreenTraces = 1;
+		Settings.LumenReflectionsScreenTraces = 0;
 	}
 };
