VPUtils: Removed unnecessary dependency on MultiUserClient and reduced it to ConcertMain.

It only depends on ConcertClient which is a ConcertMain module and not a MultiUserClient module. This effectively allows having nDisplay enabled with MultiUser disabled.

diff --git a/Engine/Plugins/Experimental/VirtualProductionUtilities/VirtualProductionUtilities.uplugin b/Engine/Plugins/Experimental/VirtualProductionUtilities/VirtualProductionUtilities.uplugin
--- a/Engine/Plugins/Experimental/VirtualProductionUtilities/VirtualProductionUtilities.uplugin
+++ b/Engine/Plugins/Experimental/VirtualProductionUtilities/VirtualProductionUtilities.uplugin
@@ -42,7 +42,7 @@
 	],
 	"Plugins": [
 		{
-			"Name": "MultiUserClient",
+			"Name": "ConcertMain",
 			"Enabled": true
 		},
 		{
