Improve editor start-up time by async loading: VPScoutingSubsystem now async loads UVPUtilitiesEditorSettings::ScoutingSubsystemEditorUtilityClassPath.

diff --git a/Engine/Plugins/Experimental/VirtualProductionUtilities/Source/VPUtilitiesEditor/Private/VPScoutingSubsystem.cpp b/Engine/Plugins/Experimental/VirtualProductionUtilities/Source/VPUtilitiesEditor/Private/VPScoutingSubsystem.cpp
--- a/Engine/Plugins/Experimental/VirtualProductionUtilities/Source/VPUtilitiesEditor/Private/VPScoutingSubsystem.cpp
+++ b/Engine/Plugins/Experimental/VirtualProductionUtilities/Source/VPUtilitiesEditor/Private/VPScoutingSubsystem.cpp
@@ -6,17 +6,14 @@
 #include "IVREditorModule.h"
 #include "VREditorMode.h"
 #include "UI/VREditorUISystem.h"
-#include "VREditorStyle.h"
-#include "WidgetBlueprint.h"
-#include "EditorUtilityActor.h"
 #include "EditorUtilityWidget.h"
 #include "Engine/AssetManager.h"
-#include "IVREditorModule.h"
+#include "LevelEditorActions.h"
 #include "UObject/ConstructorHelpers.h"
 #include "UObject/Script.h"
+#include "UObject/UObjectGlobals.h"
 #include "VPSettings.h"
 #include "VPUtilitiesEditorSettings.h"
-#include "LevelEditorActions.h"
 
 LLM_DEFINE_TAG(VirtualProductionUtilities_VPScoutingSubsystem);
 
@@ -118,15 +115,19 @@ void UVPScoutingSubsystem::Initialize(FSubsystemCollectionBase& Collection)
 
 	// Load the ScoutingHelper implemented in BP. See BaseVirtualProductionUtilitites.ini
 	VPSubsystemHelpers = nullptr;
-	if (UClass* EditorUtilityClass = GetDefault<UVPUtilitiesEditorSettings>()->ScoutingSubsystemEditorUtilityClassPath.TryLoadClass<UVPScoutingSubsystemHelpersBase>())
+	FSoftClassPath ClassPath = GetDefault<UVPUtilitiesEditorSettings>()->ScoutingSubsystemEditorUtilityClassPath;
+	ClassPath.LoadAsync(FLoadSoftObjectPathAsyncDelegate::CreateWeakLambda(this, [this](const FSoftObjectPath&, UObject* LoadedObject)
 	{
-		VPSubsystemHelpers = NewObject<UVPScoutingSubsystemHelpersBase>(GetTransientPackage(), EditorUtilityClass);
-	}
-	else
-	{
-		UE_LOG(LogVPUtilitiesEditor, Warning, TEXT("Failed loading VPScoutingHelpers \"%s\""), *GetDefault<UVPUtilitiesEditorSettings>()->ScoutingSubsystemEditorUtilityClassPath.ToString());
-	}
-
+		if (UClass* EditorUtilityClass = Cast<UClass>(LoadedObject))
+		{
+			VPSubsystemHelpers = NewObject<UVPScoutingSubsystemHelpersBase>(GetTransientPackage(), EditorUtilityClass);
+		}
+		else
+		{
+			UE_LOG(LogVPUtilitiesEditor, Warning, TEXT("Failed loading VPScoutingHelpers \"%s\""), *GetDefault<UVPUtilitiesEditorSettings>()->ScoutingSubsystemEditorUtilityClassPath.ToString());
+		}
+	}));
+	
 	// to do final initializations at the right time
 	EngineInitCompleteDelegate = FCoreDelegates::OnFEngineLoopInitComplete.AddUObject(this, &UVPScoutingSubsystem::OnEngineInitComplete);
 }
