Add CineCaptureComponent and capture function on VPBookmarkActor
Add FavoriteLevel/Flag/Timestamp properties on VPBookmarkActor
Add Sorting functions

diff --git a/Engine/Plugins/Experimental/VirtualProductionUtilities/Source/VPUtilities/Private/Actors/VPBookmarkActor.cpp b/Engine/Plugins/Experimental/VirtualProductionUtilities/Source/VPUtilities/Private/Actors/VPBookmarkActor.cpp
--- a/Engine/Plugins/Experimental/VirtualProductionUtilities/Source/VPUtilities/Private/Actors/VPBookmarkActor.cpp
+++ b/Engine/Plugins/Experimental/VirtualProductionUtilities/Source/VPUtilities/Private/Actors/VPBookmarkActor.cpp
@@ -14,6 +14,10 @@
 #include "Components/SplineMeshComponent.h"
 #include "Components/TextRenderComponent.h"
 #include "Materials/MaterialInstanceDynamic.h"
+#include "CineCameraSceneCaptureComponent.h"
+#include "Engine/TextureRenderTarget2D.h"
+#include "Engine/Texture2D.h"
+#include "Logging/MessageLog.h"
 
 AVPBookmarkActor::AVPBookmarkActor(const FObjectInitializer& ObjectInitializer)
 	: Super(ObjectInitializer)
@@ -131,6 +135,12 @@ AVPBookmarkActor::AVPBookmarkActor(const FObjectInitializer& ObjectInitializer)
 #if WITH_EDITOR
 	CameraComponent->SetCameraMesh(nullptr);// set CameraMesh to null when we are in editor view
 #endif
+
+	SceneCaptureComponent = CreateDefaultSubobject<UCineCaptureComponent2D>(TEXT("SceneCapture"));
+	SceneCaptureComponent->SetupAttachment(CameraComponent);
+	SceneCaptureComponent->bCaptureEveryFrame = false;
+	SceneCaptureComponent->bCaptureOnMovement = false;
+
 }
 
 void AVPBookmarkActor::Tick(float DeltaSeconds)
@@ -150,6 +160,65 @@ void AVPBookmarkActor::Tick(float DeltaSeconds)
 #endif
 }
 
+void AVPBookmarkActor::CaptureSnapshot()
+{
+#if WITH_EDITOR
+	if (!SceneCaptureComponent->TextureTarget)
+	{
+		UE_LOG(LogVPUtilities, Verbose, TEXT("Creating transient texture target"));
+		UTextureRenderTarget2D* NewRenderTarget2D = NewObject<UTextureRenderTarget2D>();
+		NewRenderTarget2D->RenderTargetFormat = RTF_RGBA8_SRGB;
+		
+		int32 Width, Height;
+		if (CameraComponent->AspectRatio >= 1.)
+		{
+			Width = SceneCaptureComponent->RenderTargetHighestDimension;
+			Height = FMath::Max( Width / CameraComponent->AspectRatio, 1);
+		}
+		else
+		{
+			Height = SceneCaptureComponent->RenderTargetHighestDimension;
+			Width = FMath::Max( Height * CameraComponent->AspectRatio, 1);
+		}
+		NewRenderTarget2D->InitAutoFormat(Width, Height);
+		NewRenderTarget2D->UpdateResourceImmediate(true);
+
+		SceneCaptureComponent->TextureTarget = NewRenderTarget2D;
+	}
+	if (!SnapshotTexture)
+	{
+		UE_LOG(LogVPUtilities, Verbose, TEXT("Creating transient snapshot texture"));
+		SnapshotTexture = UTexture2D::CreateTransient(SceneCaptureComponent->TextureTarget->SizeX, SceneCaptureComponent->TextureTarget->SizeY, EPixelFormat::PF_B8G8R8A8);
+	}
+
+	if (SceneCaptureComponent->TextureTarget)
+	{
+		if (!SceneCaptureComponent->TextureTarget->GameThread_GetRenderTargetResource())
+		{
+			UE_LOG(LogVPUtilities, Warning, TEXT("Unable to get render target resource"));
+			return;
+		}
+		SceneCaptureComponent->CaptureScene();
+	}
+
+	FText ErrorMessage;
+	if (!SceneCaptureComponent->TextureTarget->UpdateTexture(SnapshotTexture, /*InFlags = */CTF_Default, /*InAlphaOverride = */nullptr, /*InOnTextureChangingDelegate = */[](UTexture*) {}, &ErrorMessage))
+	{
+		FMessageLog("VPBookmarkActor").Warning(ErrorMessage);
+		return;
+	}
+	SnapshotTexture->Modify();
+	SnapshotTexture->PostEditChange();
+	SnapshotTexture->UpdateResource();
+#endif
+}
+
+void AVPBookmarkActor::UpdateTimestamp()
+{
+	Timestamp = FDateTime::UtcNow();
+}
+
+
 //VP Interaction InterfaceEvents
 
 void AVPBookmarkActor::OnBookmarkActivation_Implementation(UVPBookmark* BookmarkOut, bool bIsActive)
diff --git a/Engine/Plugins/Experimental/VirtualProductionUtilities/Source/VPUtilities/Private/Libraries/VPBlueprintLibrary.cpp b/Engine/Plugins/Experimental/VirtualProductionUtilities/Source/VPUtilities/Private/Libraries/VPBlueprintLibrary.cpp
--- a/Engine/Plugins/Experimental/VirtualProductionUtilities/Source/VPUtilities/Private/Libraries/VPBlueprintLibrary.cpp
+++ b/Engine/Plugins/Experimental/VirtualProductionUtilities/Source/VPUtilities/Private/Libraries/VPBlueprintLibrary.cpp
@@ -4,10 +4,12 @@
 
 #include "VPSettings.h"
 #include "VPUtilitiesModule.h"
+#include "Actors/VPBookmarkActor.h"
 
 #include "Engine/World.h"
 #include "HAL/IConsoleManager.h"
 #include "Components/SplineMeshComponent.h"
+#include "Misc/ComparisonUtility.h"
 
 #if WITH_EDITOR
 #include "Editor.h"
@@ -327,3 +329,29 @@ void UVPBlueprintLibrary::VPBookmarkSplineMeshIndicatorDisable(USplineMeshCompon
 {
 	SplineMesh->SetVisibility(false);
 }
+
+void UVPBlueprintLibrary::SortActorsByName(UPARAM(Ref) TArray<AActor*>& Actors, bool Ascending)
+{
+	Algo::Sort(Actors, [](const AActor* A, const AActor* B)
+		{
+			return UE::ComparisonUtility::CompareNaturalOrder(A->GetActorNameOrLabel(), B->GetActorNameOrLabel()) < 0;
+		});
+
+	if (!Ascending)
+	{
+		Algo::Reverse(Actors);
+	}
+}
+
+void UVPBlueprintLibrary::SortVPBookmarkActorsByTimestamp(UPARAM(Ref) TArray<AVPBookmarkActor*>& Actors, bool Ascending)
+{
+	Algo::Sort(Actors, [](const AVPBookmarkActor* A, const AVPBookmarkActor* B)
+		{
+			return A->Timestamp < B->Timestamp;
+		});
+
+	if (!Ascending)
+	{
+		Algo::Reverse(Actors);
+	}
+}
\ No newline at end of file
diff --git a/Engine/Plugins/Experimental/VirtualProductionUtilities/Source/VPUtilities/Public/Actors/VPBookmarkActor.h b/Engine/Plugins/Experimental/VirtualProductionUtilities/Source/VPUtilities/Public/Actors/VPBookmarkActor.h
--- a/Engine/Plugins/Experimental/VirtualProductionUtilities/Source/VPUtilities/Public/Actors/VPBookmarkActor.h
+++ b/Engine/Plugins/Experimental/VirtualProductionUtilities/Source/VPUtilities/Public/Actors/VPBookmarkActor.h
@@ -16,7 +16,7 @@ class UTextRenderComponent;
 class USplineMeshComponent;
 class UMaterialInstanceDynamic;
 class UMaterial;
-
+class UCineCaptureComponent2D;
 
 UCLASS(HideCategories=(Rendering, Lighting, HLOD, Mobile, Navigation, RayTracing, TextureStreaming), CollapseCategories)
 class VPUTILITIES_API AVPBookmarkActor : public AVPViewportTickableActorBase, public IVPInteraction, public IVPBookmarkProvider
@@ -43,10 +43,26 @@ public:
 	UPROPERTY(BlueprintReadOnly, EditDefaultsOnly, Category = "Components")
 	TObjectPtr<UCineCameraComponent> CameraComponent;
 
+	/**CineCaptureComponent*/
+	UPROPERTY(BlueprintReadOnly, EditDefaultsOnly, Category = "Components")
+	TObjectPtr<UCineCaptureComponent2D> SceneCaptureComponent;
+
 	/** Color of Bookmark in MU Session */
 	UPROPERTY(BlueprintReadWrite, EditAnywhere, Category = "Bookmark", meta = (MultiLine = "true", ExposeOnSpawn = "true"))
 	FLinearColor BookmarkColor;
 
+	/** Favorite status set by a user */
+	UPROPERTY(BlueprintReadWrite, EditAnywhere, Category = "Bookmark")
+	int32 FavoriteLevel = 0;
+
+	/** Whether this was flagged by a user */
+	UPROPERTY(BlueprintReadWrite, EditAnywhere, Category = "Bookmark")
+	bool bIsFlagged = false;
+
+	/** DateTime to store creation time */
+	UPROPERTY(VisibleAnywhere, BlueprintReadOnly, Category = "Bookmark")
+	FDateTime Timestamp;
+
 	/** Reference to Editor Bookmark UObject*/
 	UPROPERTY(BlueprintReadWrite, EditDefaultsOnly, Category = "Bookmark", meta = (MultiLine = "true"))
 	TObjectPtr<UVPBookmark> BookmarkObject;
@@ -56,7 +72,7 @@ public:
 	bool IsHome;
 
 	/** Texture reference to store render of viewpoint  */
-	UPROPERTY(BlueprintReadWrite, Category = "Snapshot")
+	UPROPERTY(BlueprintReadWrite, VisibleAnywhere, Category = "Snapshot")
 	TObjectPtr<UTexture2D> SnapshotTexture;
 
 	/**Update the mesh color and BookmarkColor variable. Intended for use with multiuser initialization*/
@@ -69,6 +85,13 @@ public:
 	UPROPERTY(EditDefaultsOnly, BlueprintReadWrite, Category = "Default")
 	TObjectPtr<UMaterialInterface> BookmarkMaterial;
 
+	/** Update SnapshotTexture captured via SceneCaptureComponent*/
+	UFUNCTION(BlueprintCallable, CallInEditor, Category = "Snapshot")
+	void CaptureSnapshot();
+
+	/** Update Timestamp property with current time*/
+	UFUNCTION(BlueprintCallable, CallInEditor, Category = "Bookmark")
+	void UpdateTimestamp();
 
 	//VPBookmark Interface events
 
diff --git a/Engine/Plugins/Experimental/VirtualProductionUtilities/Source/VPUtilities/Public/Libraries/VPBlueprintLibrary.h b/Engine/Plugins/Experimental/VirtualProductionUtilities/Source/VPUtilities/Public/Libraries/VPBlueprintLibrary.h
--- a/Engine/Plugins/Experimental/VirtualProductionUtilities/Source/VPUtilities/Public/Libraries/VPBlueprintLibrary.h
+++ b/Engine/Plugins/Experimental/VirtualProductionUtilities/Source/VPUtilities/Public/Libraries/VPBlueprintLibrary.h
@@ -12,7 +12,7 @@
 
 
 class UVPBookmark;
-
+class AVPBookmarkActor;
 
 /**
  * Functionality added to prototype the VR scouting tools
@@ -101,5 +101,12 @@ public:
 	UFUNCTION(BlueprintCallable, Category = "VPBookmarks")
 	static void VPBookmarkSplineMeshIndicatorDisable(USplineMeshComponent* SplineMesh);
 
+	/** Sort actors by actor label*/
+	UFUNCTION(BlueprintCallable, Category = "VPBookmarks")
+	static void SortActorsByName(UPARAM(Ref) TArray<AActor*>& Actors, bool Ascending = true);
+
+	/** Sort VPBookmarkActors by Timestamp parameter*/
+	UFUNCTION(BlueprintCallable, Category = "VPBookmarks")
+	static void SortVPBookmarkActorsByTimestamp(UPARAM(Ref) TArray<AVPBookmarkActor*>& Actors, bool Ascending = true);
 	
 };
\ No newline at end of file
diff --git a/Engine/Plugins/Experimental/VirtualProductionUtilities/Source/VPUtilities/VPUtilities.Build.cs b/Engine/Plugins/Experimental/VirtualProductionUtilities/Source/VPUtilities/VPUtilities.Build.cs
--- a/Engine/Plugins/Experimental/VirtualProductionUtilities/Source/VPUtilities/VPUtilities.Build.cs
+++ b/Engine/Plugins/Experimental/VirtualProductionUtilities/Source/VPUtilities/VPUtilities.Build.cs
@@ -13,6 +13,7 @@ public class VPUtilities : ModuleRules
 			{
 				"AssetRegistry",
 				"CinematicCamera",
+				"CineCameraSceneCapture",
 				"Composure",
 				"Core",
 				"CoreUObject",
diff --git a/Engine/Plugins/Experimental/VirtualProductionUtilities/VirtualProductionUtilities.uplugin b/Engine/Plugins/Experimental/VirtualProductionUtilities/VirtualProductionUtilities.uplugin
--- a/Engine/Plugins/Experimental/VirtualProductionUtilities/VirtualProductionUtilities.uplugin
+++ b/Engine/Plugins/Experimental/VirtualProductionUtilities/VirtualProductionUtilities.uplugin
@@ -80,6 +80,10 @@
 		{
 			"Name": "XRBase",
 			"Enabled": true
+		},
+		{
+			"Name": "CineCameraSceneCapture",
+			"Enabled": true
 		}
 	]
 }
