Use the PostRegisterAllComponents override to trigger the bookmark color update.  In an undo/redo flow, it is possible that that componets have not been initialized yet thus the component will be null calling the UpdateBookmarkColor method

diff --git a/Engine/Plugins/Experimental/VirtualProductionUtilities/Source/VPUtilities/Private/Actors/VPBookmarkActor.cpp b/Engine/Plugins/Experimental/VirtualProductionUtilities/Source/VPUtilities/Private/Actors/VPBookmarkActor.cpp
--- a/Engine/Plugins/Experimental/VirtualProductionUtilities/Source/VPUtilities/Private/Actors/VPBookmarkActor.cpp
+++ b/Engine/Plugins/Experimental/VirtualProductionUtilities/Source/VPUtilities/Private/Actors/VPBookmarkActor.cpp
@@ -277,7 +277,7 @@ void AVPBookmarkActor::OnActorDroppedFromTransform_Implementation()
 
 void AVPBookmarkActor::UpdateBookmarkColor(FLinearColor Color)
 {
-	if (BookmarkMeshComponent->GetStaticMesh() != nullptr)
+	if (BookmarkMeshComponent && BookmarkMeshComponent->GetStaticMesh() != nullptr)
 	{
 		UMaterialInterface* Material = BookmarkMeshComponent->GetMaterial(0);
 
@@ -315,9 +315,9 @@ void AVPBookmarkActor::PostEditChangeProperty(FPropertyChangedEvent& PropertyCha
 }
 #endif
 
-void AVPBookmarkActor::OnConstruction(const FTransform& Transform)
+void AVPBookmarkActor::PostRegisterAllComponents()
 {
-	Super::OnConstruction(Transform);
+	Super::PostRegisterAllComponents();
 
 	AVPBookmarkActor::UpdateBookmarkColor(BookmarkColor);
 }
diff --git a/Engine/Plugins/Experimental/VirtualProductionUtilities/Source/VPUtilities/Public/Actors/VPBookmarkActor.h b/Engine/Plugins/Experimental/VirtualProductionUtilities/Source/VPUtilities/Public/Actors/VPBookmarkActor.h
--- a/Engine/Plugins/Experimental/VirtualProductionUtilities/Source/VPUtilities/Public/Actors/VPBookmarkActor.h
+++ b/Engine/Plugins/Experimental/VirtualProductionUtilities/Source/VPUtilities/Public/Actors/VPBookmarkActor.h
@@ -124,7 +124,7 @@ public:
 
 	//Overrides
 
-	virtual void OnConstruction(const FTransform& Transform) override;
+	virtual void PostRegisterAllComponents() override;
 
 	virtual void Tick(float DeltaSeconds) override;
 
