VPBookmarkActor: CaptureSnapshot to create opaque texture so that it shows up in projects with alpha enabled

diff --git a/Engine/Plugins/Experimental/VirtualProductionUtilities/Source/VPUtilities/Private/Actors/VPBookmarkActor.cpp b/Engine/Plugins/Experimental/VirtualProductionUtilities/Source/VPUtilities/Private/Actors/VPBookmarkActor.cpp
--- a/Engine/Plugins/Experimental/VirtualProductionUtilities/Source/VPUtilities/Private/Actors/VPBookmarkActor.cpp
+++ b/Engine/Plugins/Experimental/VirtualProductionUtilities/Source/VPUtilities/Private/Actors/VPBookmarkActor.cpp
@@ -202,7 +202,7 @@ void AVPBookmarkActor::CaptureSnapshot()
 	}
 
 	FText ErrorMessage;
-	if (!SceneCaptureComponent->TextureTarget->UpdateTexture(SnapshotTexture, /*InFlags = */CTF_Default, /*InAlphaOverride = */nullptr, /*InOnTextureChangingDelegate = */[](UTexture*) {}, &ErrorMessage))
+	if (!SceneCaptureComponent->TextureTarget->UpdateTexture(SnapshotTexture, /*InFlags = */static_cast<EConstructTextureFlags>(CTF_Default | CTF_ForceOpaque), /*InAlphaOverride = */nullptr, /*InOnTextureChangingDelegate = */[](UTexture*) {}, &ErrorMessage))
 	{
 		FMessageLog("VPBookmarkActor").Warning(ErrorMessage);
 		return;
