VP Utilities

- Moved all Legacy VScout assets to a "DEPRECATED" folder.
- Fixed up any text-based references to those assets in their new folder.

[RN] Legacy Virtual Scouting toolset is marked as deprecated

#virtualized

diff --git a/Engine/Plugins/Experimental/VirtualProductionUtilities/Config/BaseVirtualProductionUtilities.ini b/Engine/Plugins/Experimental/VirtualProductionUtilities/Config/BaseVirtualProductionUtilities.ini
--- a/Engine/Plugins/Experimental/VirtualProductionUtilities/Config/BaseVirtualProductionUtilities.ini
+++ b/Engine/Plugins/Experimental/VirtualProductionUtilities/Config/BaseVirtualProductionUtilities.ini
@@ -21,10 +21,10 @@
 +EnumRedirects=(OldName="/Script/VPUtilities.EVPWidgetDisplayType", NewName="/Script/VPUtilities.EVPWidgetDisplayType", ValueChanges=(("PostProcess", "PostProcessWithBlendMaterial")))
 
 [/Script/VPUtilitiesEditor.VPUtilitiesEditorSettings]
-ScoutingSubsystemEditorUtilityClassPath=/VirtualProductionUtilities/BP_VPScoutingSubsystemHelpers.BP_VPScoutingSubsystemHelpers_C
-GestureManagerEditorUtilityClassPath=/VirtualProductionUtilities/VRTools/BP_GestureManager.BP_GestureManager_C
-+AdditionnalClassToLoad=/VirtualProductionUtilities/VR/VirtualScoutingInteractor.VirtualScoutingInteractor_C
-+AdditionnalClassToLoad=/VirtualProductionUtilities/VR/VirtualScoutingTeleporter.VirtualScoutingTeleporter_C
+ScoutingSubsystemEditorUtilityClassPath=/VirtualProductionUtilities/DEPRECATED/BP_VPScoutingSubsystemHelpers.BP_VPScoutingSubsystemHelpers_C
+GestureManagerEditorUtilityClassPath=/VirtualProductionUtilities/DEPRECATED/VRTools/BP_GestureManager.BP_GestureManager_C
++AdditionnalClassToLoad=/VirtualProductionUtilities/DEPRECATED/VR/VirtualScoutingInteractor.VirtualScoutingInteractor_C
++AdditionnalClassToLoad=/VirtualProductionUtilities/DEPRECATED/VR/VirtualScoutingTeleporter.VirtualScoutingTeleporter_C
 
 
 [/Script/VPUtilities.VPBookmarkSettings]
diff --git a/Engine/Plugins/VirtualProduction/VirtualScouting/Source/VirtualScoutingEditor/Private/VPScoutingMode.cpp b/Engine/Plugins/VirtualProduction/VirtualScouting/Source/VirtualScoutingEditor/Private/VPScoutingMode.cpp
--- a/Engine/Plugins/VirtualProduction/VirtualScouting/Source/VirtualScoutingEditor/Private/VPScoutingMode.cpp
+++ b/Engine/Plugins/VirtualProduction/VirtualScouting/Source/VirtualScoutingEditor/Private/VPScoutingMode.cpp
@@ -24,8 +24,8 @@ static const FName OpenXRSystemName = FName(TEXT("OpenXR"));
 
 UVPScoutingMode::UVPScoutingMode()
 {
-	InteractorClass = FSoftObjectPath("/VirtualProductionUtilities/VR/VirtualScoutingInteractor.VirtualScoutingInteractor_C");
-	TeleporterClass = FSoftObjectPath("/VirtualProductionUtilities/VR/VirtualScoutingTeleporter.VirtualScoutingTeleporter_C");
+	InteractorClass = FSoftObjectPath("/VirtualProductionUtilities/DEPRECATED/VR/VirtualScoutingInteractor.VirtualScoutingInteractor_C");
+	TeleporterClass = FSoftObjectPath("/VirtualProductionUtilities/DEPRECATED/VR/VirtualScoutingTeleporter.VirtualScoutingTeleporter_C");
 }
 
 
