[Legacy Virtual Scouting]

Code deprecation and removal from Virtual Scouting and Virtual Production Utilities Plugins

[RN] Virtual Scouting - removed code used exclusively by the Virtual Scouting 1.0 tools which are now deprecated

diff --git a/Engine/Plugins/Experimental/VirtualProductionUtilities/Config/BaseVirtualProductionUtilities.ini b/Engine/Plugins/Experimental/VirtualProductionUtilities/Config/BaseVirtualProductionUtilities.ini
--- a/Engine/Plugins/Experimental/VirtualProductionUtilities/Config/BaseVirtualProductionUtilities.ini
+++ b/Engine/Plugins/Experimental/VirtualProductionUtilities/Config/BaseVirtualProductionUtilities.ini
@@ -20,13 +20,6 @@
 +PropertyRedirects=(OldName="VPFullScreenUserWidget.PostProcessDisplayType",NewName="VPFullScreenUserWidget.PostProcessDisplayTypeWithBlendMaterial")
 +EnumRedirects=(OldName="/Script/VPUtilities.EVPWidgetDisplayType", NewName="/Script/VPUtilities.EVPWidgetDisplayType", ValueChanges=(("PostProcess", "PostProcessWithBlendMaterial")))
 
-[/Script/VPUtilitiesEditor.VPUtilitiesEditorSettings]
-ScoutingSubsystemEditorUtilityClassPath=/VirtualProductionUtilities/DEPRECATED/BP_VPScoutingSubsystemHelpers.BP_VPScoutingSubsystemHelpers_C
-GestureManagerEditorUtilityClassPath=/VirtualProductionUtilities/DEPRECATED/VRTools/BP_GestureManager.BP_GestureManager_C
-+AdditionnalClassToLoad=/VirtualProductionUtilities/DEPRECATED/VR/VirtualScoutingInteractor.VirtualScoutingInteractor_C
-+AdditionnalClassToLoad=/VirtualProductionUtilities/DEPRECATED/VR/VirtualScoutingTeleporter.VirtualScoutingTeleporter_C
-
-
 [/Script/VPUtilities.VPBookmarkSettings]
 
 BookmarkMeshPath=/VirtualProductionUtilities/Bookmark/BookmarkVPMesh.BookmarkVPMesh
diff --git a/Engine/Plugins/Experimental/VirtualProductionUtilities/Source/VPUtilities/Private/Libraries/VPCameraBlueprintLibrary.cpp b/Engine/Plugins/Experimental/VirtualProductionUtilities/Source/VPUtilities/Private/Libraries/VPCameraBlueprintLibrary.cpp
--- a/Engine/Plugins/Experimental/VirtualProductionUtilities/Source/VPUtilities/Private/Libraries/VPCameraBlueprintLibrary.cpp
+++ b/Engine/Plugins/Experimental/VirtualProductionUtilities/Source/VPUtilities/Private/Libraries/VPCameraBlueprintLibrary.cpp
@@ -16,7 +16,7 @@
 #include "Engine/Selection.h"
 #endif
 
-
+PRAGMA_DISABLE_DEPRECATION_WARNINGS
 FVPCameraRigSpawnParams::FVPCameraRigSpawnParams()
 	: bUseWorldSpace(true)
 	, bUseFirstPointAsSpawnLocation(false)
@@ -236,3 +236,4 @@ ACameraRig_Rail* UVPCameraBlueprintLibrary::SpawnCameraRigFromSelectedActors(UOb
 	UE_LOG(LogVPUtilities, Warning, TEXT("UVPCameraBlueprintLibrary::SpawnCameraRigFromSelectedActors - Only callable from editor"));
 	return nullptr;
 }
+PRAGMA_ENABLE_DEPRECATION_WARNINGS
\ No newline at end of file
diff --git a/Engine/Plugins/Experimental/VirtualProductionUtilities/Source/VPUtilities/Public/IVPContextMenuProvider.h b/Engine/Plugins/Experimental/VirtualProductionUtilities/Source/VPUtilities/Public/IVPContextMenuProvider.h
--- a/Engine/Plugins/Experimental/VirtualProductionUtilities/Source/VPUtilities/Public/IVPContextMenuProvider.h
+++ b/Engine/Plugins/Experimental/VirtualProductionUtilities/Source/VPUtilities/Public/IVPContextMenuProvider.h
@@ -9,7 +9,6 @@
 
 class UVPContextMenu;
 
-
 UINTERFACE(BlueprintType)
 class UVPContextMenuProvider : public UInterface
 {
@@ -22,6 +21,7 @@ class IVPContextMenuProvider : public IInterface
 	GENERATED_BODY()
 
 public:
-	UFUNCTION(BlueprintNativeEvent, BlueprintCallable, CallInEditor, Category = "Virtual Production")
+	UE_DEPRECATED(5.5,"Code will be removed from UE5.7")
+	UFUNCTION(BlueprintNativeEvent, BlueprintCallable, CallInEditor, Category = "Virtual Production", meta=(DeprecatedFunction, DeprecationMessage="This function will be removed from UE5.7"))
 	void OnCreateContextMenu();
 };
diff --git a/Engine/Plugins/Experimental/VirtualProductionUtilities/Source/VPUtilities/Public/IVPInteraction.h b/Engine/Plugins/Experimental/VirtualProductionUtilities/Source/VPUtilities/Public/IVPInteraction.h
--- a/Engine/Plugins/Experimental/VirtualProductionUtilities/Source/VPUtilities/Public/IVPInteraction.h
+++ b/Engine/Plugins/Experimental/VirtualProductionUtilities/Source/VPUtilities/Public/IVPInteraction.h
@@ -9,6 +9,7 @@
 class UVirtualScoutingInteraction;
 
 // This class does not need to be modified.
+
 UINTERFACE(BlueprintType)
 class UVPInteraction : public UInterface
 {
@@ -23,14 +24,17 @@ class IVPInteraction : public IInterface
 public:
 
 	/** Called when actor is dropped from Carry state by the VREd interactor */
-	UFUNCTION(BlueprintNativeEvent, CallInEditor, BlueprintCallable, Category = "VirtualScouting|Interaction")
+	UE_DEPRECATED(5.5,"Code will be removed from UE5.7")
+	UFUNCTION(BlueprintNativeEvent, CallInEditor, BlueprintCallable, Category = "VirtualScouting|Interaction", meta=(DeprecatedFunction, DeprecationMessage="This function will be removed from UE5.7"))
 	void OnActorDroppedFromCarry();
 
 	/** Called when an actor is selected by the VREd interactor */
-	UFUNCTION(BlueprintNativeEvent, CallInEditor, BlueprintCallable, Category = "VirtualScouting|Interaction")
+	UE_DEPRECATED(5.5,"Code will be removed from UE5.7")
+	UFUNCTION(BlueprintNativeEvent, CallInEditor, BlueprintCallable, Category = "VirtualScouting|Interaction", meta=(DeprecatedFunction, DeprecationMessage="This function will be removed from UE5.7"))
 	void OnActorSelectedForTransform();
 
 	/** Called when dropped by VRTool */
-	UFUNCTION(BlueprintNativeEvent, CallInEditor, BlueprintCallable, Category = "VirtualScouting|Interaction")
+	UE_DEPRECATED(5.5,"Code will be removed from UE5.7")
+	UFUNCTION(BlueprintNativeEvent, CallInEditor, BlueprintCallable, Category = "VirtualScouting|Interaction", meta=(DeprecatedFunction, DeprecationMessage="This function will be removed from UE5.7"))
 	void OnActorDroppedFromTransform();
 };
\ No newline at end of file
diff --git a/Engine/Plugins/Experimental/VirtualProductionUtilities/Source/VPUtilities/Public/Libraries/VPCameraBlueprintLibrary.h b/Engine/Plugins/Experimental/VirtualProductionUtilities/Source/VPUtilities/Public/Libraries/VPCameraBlueprintLibrary.h
--- a/Engine/Plugins/Experimental/VirtualProductionUtilities/Source/VPUtilities/Public/Libraries/VPCameraBlueprintLibrary.h
+++ b/Engine/Plugins/Experimental/VirtualProductionUtilities/Source/VPUtilities/Public/Libraries/VPCameraBlueprintLibrary.h
@@ -11,17 +11,17 @@
 class ACameraRig_Rail;
 
 UENUM(BlueprintType)
-enum class EVPCameraRigSpawnLinearApproximationMode : uint8
+enum class UE_DEPRECATED(5.5,"Code will be removed from UE5.7") EVPCameraRigSpawnLinearApproximationMode : uint8
 {
 	None			UMETA(Display = "No Approximation"),	// We won't do linear approximation, instead using the Spline as constructed initially.
 	Density,												// LinearApproximationParam will be used as a density value 
 	IntegrationStep,										// LinearApproximationParam will be used as the Integration step in Unreal Units.
 };
 
-
+PRAGMA_DISABLE_DEPRECATION_WARNINGS
 /** Parameters used to custom the CameraRig that's created. */
 USTRUCT(BlueprintType)
-struct FVPCameraRigSpawnParams
+struct UE_DEPRECATED(5.5,"Code will be removed from UE5.7") FVPCameraRigSpawnParams
 {
 	GENERATED_BODY();
 
@@ -29,21 +29,24 @@ struct FVPCameraRigSpawnParams
 
 public:
 	/** Use world space (as opposed to local space) for points. */
-	UPROPERTY(Transient, BlueprintReadWrite, Category="Camera Rig")
+	UE_DEPRECATED(5.5,"This property will be removed from UE5.7")
+	UPROPERTY(Transient, BlueprintReadWrite, Category="Camera Rig",  meta=(DeprecatedProperty, DeprecationMessage="This property deprecated"))
 	bool bUseWorldSpace;
 
 	/**
 	 * Use the first vector of input as the spawn transform.
 	 * Causes RigTransform to be completely ignored.
 	 */
-	UPROPERTY(Transient, BlueprintReadWrite, Category = "Camera Rig")
+	UE_DEPRECATED(5.5,"This property will be removed from UE5.7")
+	UPROPERTY(Transient, BlueprintReadWrite, Category = "Camera Rig",  meta=(DeprecatedProperty, DeprecationMessage="This property deprecated"))
 	bool bUseFirstPointAsSpawnLocation;
 
 	/**
 	 * Causes a linear approximation of the spline points to be generated instead
 	 * of relying purely on the passed in points / curves.
 	 */
-	UPROPERTY(Transient, BlueprintReadWrite, Category = "Camera Rig")
+	UE_DEPRECATED(5.5,"This property will be removed from UE5.7")
+	UPROPERTY(Transient, BlueprintReadWrite, Category = "Camera Rig",  meta=(DeprecatedProperty, DeprecationMessage="This property deprecated"))
 	EVPCameraRigSpawnLinearApproximationMode LinearApproximationMode;
 
 	/**
@@ -56,30 +59,35 @@ public:
 	 * increase the number of spline points and will therefore increase the accuracy
 	 * (at the cost of increased complexity).
 	 */
-	UPROPERTY(Transient, BlueprintReadWrite, Category = "Camera Rig")
+	UE_DEPRECATED(5.5,"This property will be removed from UE5.7")
+	UPROPERTY(Transient, BlueprintReadWrite, Category = "Camera Rig",  meta=(DeprecatedProperty, DeprecationMessage="This property deprecated"))
 	float LinearApproximationParam;
 };
 
-
 UCLASS()
-class VPUTILITIES_API UVPCameraBlueprintLibrary : public UBlueprintFunctionLibrary
+class UE_DEPRECATED(5.5,"Code will be removed from UE5.7") VPUTILITIES_API UVPCameraBlueprintLibrary : public UBlueprintFunctionLibrary
 {
 	GENERATED_BODY()
 
 public:
-
-	UFUNCTION(BlueprintCallable, Category = "Virtual Production")
+	UE_DEPRECATED(5.5,"Code will be removed from UE5.7")
+	UFUNCTION(BlueprintCallable, Category = "Virtual Production", meta=(DeprecatedFunction, DeprecationMessage="This function will be removed from UE5.7"))
 	static ACameraRig_Rail* SpawnDollyTrackFromPoints(UObject* WorldContextObject, UPARAM(ref) const TArray<FTransform>& Points, ESplinePointType::Type InterpType = ESplinePointType::Linear);
 
-	UFUNCTION(BlueprintCallable, Category = "Virtual Production")
+	UE_DEPRECATED(5.5,"Code will be removed from UE5.7")
+	UFUNCTION(BlueprintCallable, Category = "Virtual Production", meta=(DeprecatedFunction, DeprecationMessage="This function will be removed from UE5.7"))
 	static ACameraRig_Rail* SpawnDollyTrackFromPointsSmooth(UObject* WorldContextObject, UPARAM(ref) const TArray<FTransform>& Points, ESplinePointType::Type InterpType = ESplinePointType::Linear);
 
-	UFUNCTION(BlueprintCallable, Category = "Virtual Production")
+	UE_DEPRECATED(5.5,"Code will be removed from UE5.7")
+	UFUNCTION(BlueprintCallable, Category = "Virtual Production", meta=(DeprecatedFunction, DeprecationMessage="This function will be removed from UE5.7"))
 	static ACameraRig_Rail* SpawnCameraRigFromActors(UObject* WorldContextObject, const FTransform& RigTransform, const TArray<AActor*>& Actors, const FVPCameraRigSpawnParams& Params);
 
-	UFUNCTION(BlueprintCallable, Category = "Virtual Production")
+	UE_DEPRECATED(5.5,"Code will be removed from UE5.7")
+	UFUNCTION(BlueprintCallable, Category = "Virtual Production", meta=(DeprecatedFunction, DeprecationMessage="This function will be removed from UE5.7"))
 	static ACameraRig_Rail* SpawnCameraRigFromPoints(UObject* WorldContextObject, const FTransform& RigTransform, const TArray<FVector>& Points, const FVPCameraRigSpawnParams& Params);
 
-	UFUNCTION(BlueprintCallable, Category = "Virtual Production")
+	UE_DEPRECATED(5.5,"Code will be removed from UE5.7")
+	UFUNCTION(BlueprintCallable, Category = "Virtual Production", meta=(DeprecatedFunction, DeprecationMessage="This function will be removed from UE5.7"))
 	static ACameraRig_Rail* SpawnCameraRigFromSelectedActors(UObject* WorldContextObject, const FTransform& RigTransform, const FVPCameraRigSpawnParams& Params);
 };
+PRAGMA_ENABLE_DEPRECATION_WARNINGS
\ No newline at end of file
diff --git a/Engine/Plugins/Experimental/VirtualProductionUtilities/Source/VPUtilitiesEditor/Private/VPCameraBase.cpp b/Engine/Plugins/Experimental/VirtualProductionUtilities/Source/VPUtilitiesEditor/Private/VPCameraBase.cpp
deleted file mode 100644
diff --git a/Engine/Plugins/Experimental/VirtualProductionUtilities/Source/VPUtilitiesEditor/Private/VPCameraBase.h b/Engine/Plugins/Experimental/VirtualProductionUtilities/Source/VPUtilitiesEditor/Private/VPCameraBase.h
deleted file mode 100644
diff --git a/Engine/Plugins/Experimental/VirtualProductionUtilities/Source/VPUtilitiesEditor/Private/VPCameraUIBase.cpp b/Engine/Plugins/Experimental/VirtualProductionUtilities/Source/VPUtilitiesEditor/Private/VPCameraUIBase.cpp
deleted file mode 100644
diff --git a/Engine/Plugins/Experimental/VirtualProductionUtilities/Source/VPUtilitiesEditor/Private/VPCameraUIBase.h b/Engine/Plugins/Experimental/VirtualProductionUtilities/Source/VPUtilitiesEditor/Private/VPCameraUIBase.h
deleted file mode 100644
diff --git a/Engine/Plugins/Experimental/VirtualProductionUtilities/Source/VPUtilitiesEditor/Private/VPCustomUIHandler.cpp b/Engine/Plugins/Experimental/VirtualProductionUtilities/Source/VPUtilitiesEditor/Private/VPCustomUIHandler.cpp
deleted file mode 100644
diff --git a/Engine/Plugins/Experimental/VirtualProductionUtilities/Source/VPUtilitiesEditor/Private/VPCustomUIHandler.h b/Engine/Plugins/Experimental/VirtualProductionUtilities/Source/VPUtilitiesEditor/Private/VPCustomUIHandler.h
deleted file mode 100644
diff --git a/Engine/Plugins/Experimental/VirtualProductionUtilities/Source/VPUtilitiesEditor/Private/VPScoutingSubsystem.cpp b/Engine/Plugins/Experimental/VirtualProductionUtilities/Source/VPUtilitiesEditor/Private/VPScoutingSubsystem.cpp
--- a/Engine/Plugins/Experimental/VirtualProductionUtilities/Source/VPUtilitiesEditor/Private/VPScoutingSubsystem.cpp
+++ b/Engine/Plugins/Experimental/VirtualProductionUtilities/Source/VPUtilitiesEditor/Private/VPScoutingSubsystem.cpp
@@ -17,6 +17,7 @@
 
 LLM_DEFINE_TAG(VirtualProductionUtilities_VPScoutingSubsystem);
 
+PRAGMA_DISABLE_DEPRECATION_WARNINGS
 /* UVPScoutingSubsystem name
  *****************************************************************************/
 const FName UVPScoutingSubsystem::VProdPanelID = FName(TEXT("VirtualProductionPanel"));
@@ -426,3 +427,4 @@ void UVPScoutingSubsystem::ToggleRotationGridSnapping()
 {
 	FLevelEditorActionCallbacks::RotationGridSnap_Clicked();
 }
+PRAGMA_ENABLE_DEPRECATION_WARNINGS
\ No newline at end of file
diff --git a/Engine/Plugins/Experimental/VirtualProductionUtilities/Source/VPUtilitiesEditor/Private/VPUIBase.cpp b/Engine/Plugins/Experimental/VirtualProductionUtilities/Source/VPUtilitiesEditor/Private/VPUIBase.cpp
deleted file mode 100644
diff --git a/Engine/Plugins/Experimental/VirtualProductionUtilities/Source/VPUtilitiesEditor/Private/VPUIBase.h b/Engine/Plugins/Experimental/VirtualProductionUtilities/Source/VPUtilitiesEditor/Private/VPUIBase.h
deleted file mode 100644
diff --git a/Engine/Plugins/Experimental/VirtualProductionUtilities/Source/VPUtilitiesEditor/Private/VPUtilitiesEditorModule.cpp b/Engine/Plugins/Experimental/VirtualProductionUtilities/Source/VPUtilitiesEditor/Private/VPUtilitiesEditorModule.cpp
--- a/Engine/Plugins/Experimental/VirtualProductionUtilities/Source/VPUtilitiesEditor/Private/VPUtilitiesEditorModule.cpp
+++ b/Engine/Plugins/Experimental/VirtualProductionUtilities/Source/VPUtilitiesEditor/Private/VPUtilitiesEditorModule.cpp
@@ -40,9 +40,6 @@ void FVPUtilitiesEditorModule::StartupModule()
 
 	FVPUtilitiesEditorStyle::Register();
 
-	CustomUIHandler.Reset(NewObject<UVPCustomUIHandler>());
-	CustomUIHandler->Init();
-
 	SGenlockProviderTab::RegisterNomadTabSpawner(WorkspaceMenu::GetMenuStructure().GetLevelEditorVirtualProductionCategory());
 	STimecodeProviderTab::RegisterNomadTabSpawner();
 
@@ -62,13 +59,6 @@ void FVPUtilitiesEditorModule::ShutdownModule()
 	STimecodeProviderTab::UnregisterNomadTabSpawner();
 	SGenlockProviderTab::UnregisterNomadTabSpawner();
 
-	if (UObjectInitialized())
-	{
-		CustomUIHandler->Uninit();
-	}
-
-	CustomUIHandler.Reset();
-
 	FVPUtilitiesEditorStyle::Unregister();
 }
 
@@ -208,12 +198,7 @@ bool FVPUtilitiesEditorModule::OnSettingsModified()
 	{
 		OSCServer->Stop();
 	}
-
-	IConsoleVariable* GizmoCVar = IConsoleManager::Get().FindConsoleVariable(TEXT("VI.ShowTransformGizmo"));
-	GizmoCVar->Set(Settings->bUseTransformGizmo);
-
-	IConsoleVariable* InertiaCVar = IConsoleManager::Get().FindConsoleVariable(TEXT("VI.HighSpeedInertiaDamping"));
-	InertiaCVar->Set(Settings->bUseGripInertiaDamping ? Settings->InertiaDamping : 0);
+	
 	return true;
 }
 
diff --git a/Engine/Plugins/Experimental/VirtualProductionUtilities/Source/VPUtilitiesEditor/Private/VPUtilitiesEditorModule.h b/Engine/Plugins/Experimental/VirtualProductionUtilities/Source/VPUtilitiesEditor/Private/VPUtilitiesEditorModule.h
--- a/Engine/Plugins/Experimental/VirtualProductionUtilities/Source/VPUtilitiesEditor/Private/VPUtilitiesEditorModule.h
+++ b/Engine/Plugins/Experimental/VirtualProductionUtilities/Source/VPUtilitiesEditor/Private/VPUtilitiesEditorModule.h
@@ -5,7 +5,7 @@
 #include "CoreMinimal.h"
 #include "IVPUtilitiesEditorModule.h"
 #include "UObject/StrongObjectPtr.h"
-#include "VPCustomUIHandler.h"
+
 
 DECLARE_LOG_CATEGORY_EXTERN(LogVPUtilitiesEditor, Log, Log);
 
@@ -49,7 +49,5 @@ private:
 
 	/** Virtual production role identifier for the notification bar. */
 	static const FName VPRoleNotificationBarIdentifier;
-
-	/** UI Handler for virtual scouting. */
-	TStrongObjectPtr<UVPCustomUIHandler> CustomUIHandler;
+	
 };
\ No newline at end of file
diff --git a/Engine/Plugins/Experimental/VirtualProductionUtilities/Source/VPUtilitiesEditor/Private/VRTool.cpp b/Engine/Plugins/Experimental/VirtualProductionUtilities/Source/VPUtilitiesEditor/Private/VRTool.cpp
--- a/Engine/Plugins/Experimental/VirtualProductionUtilities/Source/VPUtilitiesEditor/Private/VRTool.cpp
+++ b/Engine/Plugins/Experimental/VirtualProductionUtilities/Source/VPUtilitiesEditor/Private/VRTool.cpp
@@ -4,6 +4,7 @@
 #include "GameFramework/InputSettings.h"
 #include "Engine/InputDelegateBinding.h"
 
+PRAGMA_DISABLE_DEPRECATION_WARNINGS
 
 AVRTool::AVRTool(const FObjectInitializer& ObjectInitializer)
 	: Super(ObjectInitializer)
@@ -64,4 +65,5 @@ void AVRTool::SetReceivesEditorInput(bool bInValue)
 	{
 		RemoveEditorInput();
 	}
-}
\ No newline at end of file
+}
+PRAGMA_ENABLE_DEPRECATION_WARNINGS
\ No newline at end of file
diff --git a/Engine/Plugins/Experimental/VirtualProductionUtilities/Source/VPUtilitiesEditor/Public/VPScoutingSubsystem.h b/Engine/Plugins/Experimental/VirtualProductionUtilities/Source/VPUtilitiesEditor/Public/VPScoutingSubsystem.h
--- a/Engine/Plugins/Experimental/VirtualProductionUtilities/Source/VPUtilitiesEditor/Public/VPScoutingSubsystem.h
+++ b/Engine/Plugins/Experimental/VirtualProductionUtilities/Source/VPUtilitiesEditor/Public/VPScoutingSubsystem.h
@@ -27,8 +27,9 @@ enum class EVProdPanelIDs : uint8
 /*
  * Base class of the helper class defined in BP
  */
+
 UCLASS(Abstract, Blueprintable, MinimalAPI, meta = (ShowWorldContextPin))
-class UVPScoutingSubsystemHelpersBase : public UObject
+class UE_DEPRECATED(5.5,"Code will be removed from UE5.7") UVPScoutingSubsystemHelpersBase : public UObject
 {
 	GENERATED_BODY()
 };
@@ -37,20 +38,25 @@ class UVPScoutingSubsystemHelpersBase : public UObject
 /*
  * Base class of the gesture manager defined in BP
  */
+
 UCLASS(Abstract, Blueprintable, MinimalAPI, meta = (ShowWorldContextPin))
-class UVPScoutingSubsystemGestureManagerBase : public UObject, public FTickableEditorObject
+class UE_DEPRECATED(5.5,"This class will be removed from UE5.7") UVPScoutingSubsystemGestureManagerBase : public UObject, public FTickableEditorObject
 {
 	GENERATED_BODY()
 
 public:
 	UVPScoutingSubsystemGestureManagerBase();
 
-	UFUNCTION(BlueprintNativeEvent, CallInEditor, BlueprintCallable, Category = "Tick")
+	UE_DEPRECATED(5.5,"Code will be removed from UE5.7")
+	UFUNCTION(BlueprintNativeEvent, CallInEditor, BlueprintCallable, Category = "Tick", meta=(DeprecatedFunction, DeprecationMessage="This function will be removed from UE5.7"))
 	void EditorTick(float DeltaSeconds);
 
-	UFUNCTION(BlueprintNativeEvent, CallInEditor, BlueprintCallable, Category = "VR")
+	UE_DEPRECATED(5.5,"Code will be removed from UE5.7")
+	UFUNCTION(BlueprintNativeEvent, CallInEditor, BlueprintCallable, Category = "VR", meta=(DeprecatedFunction, DeprecationMessage="This function will be removed from UE5.7"))
 	void OnVREditingModeEnter();
-	UFUNCTION(BlueprintNativeEvent, CallInEditor, BlueprintCallable, Category = "VR")
+	
+	UE_DEPRECATED(5.5,"Code will be removed from UE5.7")
+	UFUNCTION(BlueprintNativeEvent, CallInEditor, BlueprintCallable, Category = "VR", meta=(DeprecatedFunction, DeprecationMessage="This function will be removed from UE5.7"))
 	void OnVREditingModeExit();
 
 	//~ Begin UObject interface
@@ -73,8 +79,10 @@ private:
 /*
  * Subsystem used for VR Scouting
  */
+
+PRAGMA_DISABLE_DEPRECATION_WARNINGS
 UCLASS()
-class UVPScoutingSubsystem : public UEditorSubsystem
+class UE_DEPRECATED(5.5,"This class will be removed from UE5.7") UVPScoutingSubsystem : public UEditorSubsystem
 {
 	GENERATED_BODY()
 
@@ -85,46 +93,57 @@ public:
 	virtual void Deinitialize() override;
 
 	/** Subsystems can't have any Blueprint implementations, so we attach this class for any BP logic that we to provide. */
-	UPROPERTY(Transient, BlueprintReadOnly, Category = "Virtual Production")
+	UE_DEPRECATED(5.5,"Property will be removed from UE5.7")
+	UPROPERTY(Transient, BlueprintReadOnly, Category = "Virtual Production",  meta=(DeprecatedProperty, DeprecationMessage="This property is deprecated"))
 	TObjectPtr<UVPScoutingSubsystemHelpersBase> VPSubsystemHelpers;
 
 	/** GestureManager that manage some user input in VR editor. */
-	UPROPERTY(Transient, BlueprintReadOnly, Category = "Virtual Production")
+	UE_DEPRECATED(5.5,"Property will be removed from UE5.7")
+	UPROPERTY(Transient, BlueprintReadOnly, Category = "Virtual Production", meta=(DeprecatedProperty, DeprecationMessage="This property is deprecated"))
 	TObjectPtr<UVPScoutingSubsystemGestureManagerBase> GestureManager;
 
 	/** bool to keep track of whether the settings menu panel in the main menu is open*/
-	UPROPERTY(EditAnywhere, BlueprintReadWrite, Category = "Main Menu")
+	UE_DEPRECATED(5.5,"Property will be removed from UE5.7")
+	UPROPERTY(EditAnywhere, BlueprintReadWrite, Category = "Main Menu", meta=(DeprecatedProperty, DeprecationMessage="This property is deprecated"))
 	bool IsSettingsMenuOpen;
 
 	/** This is a multiplier for grip nav speed so we can keep the grip nav value in the range 0-1 and increase this variable if we need a bigger range */
-	UPROPERTY(EditAnywhere, BlueprintReadWrite, Category = "Virtual Production")
+	UE_DEPRECATED(5.5,"Property will be removed from UE5.7")
+	UPROPERTY(EditAnywhere, BlueprintReadWrite, Category = "Virtual Production", meta=(DeprecatedProperty, DeprecationMessage="This property is deprecated"))
 	float GripNavSpeedCoeff;
 
 	// @todo: Guard against user-created name collisions
 	/** Open a widget UI in front of the user. Opens default VProd UI (defined via the 'Virtual Scouting User Interface' setting) if null. */
-	UFUNCTION(BlueprintCallable, Category = "Virtual Production")
+	UE_DEPRECATED(5.5,"Function will be removed from UE5.7")
+	UFUNCTION(BlueprintCallable, Category = "Virtual Production", meta=(DeprecatedFunction, DeprecationMessage="This function will be removed from UE5.7"))
 	void ToggleVRScoutingUI(UPARAM(ref) FVREditorFloatingUICreationContext& CreationContext);
 
 	/** Hide VR Sequencer Window*/
-	UFUNCTION(BlueprintCallable, Category = "Virtual Production")
+	UE_DEPRECATED(5.5,"Function will be removed from UE5.7")
+	UFUNCTION(BlueprintCallable, Category = "Virtual Production", meta=(DeprecatedFunction, DeprecationMessage="This function will be removed from UE5.7"))
 	void HideInfoDisplayPanel();
 
 	/** Check whether a widget UI is open*/
-	UFUNCTION(BlueprintPure, Category = "Virtual Production")
+	UE_DEPRECATED(5.5,"Function will be removed from UE5.7")
+	UFUNCTION(BlueprintPure, Category = "Virtual Production", meta=(DeprecatedFunction, DeprecationMessage="This function will be removed from UE5.7"))
 	bool IsVRScoutingUIOpen(const FName& PanelID);
 
 	/** Get UI panel Actor from the passed ID */
-	UFUNCTION(BlueprintPure, Category = "Virtual Production")
+	UE_DEPRECATED(5.5,"Function will be removed from UE5.7")
+	UFUNCTION(BlueprintPure, Category = "Virtual Production", meta=(DeprecatedFunction, DeprecationMessage="This function will be removed from UE5.7"))
 	AVREditorFloatingUI* GetPanelActor(const FName& PanelID) const;
 
 	/** Get UI panel widget from the passed ID */
-	UFUNCTION(BlueprintPure, Category = "Virtual Production")
+	UE_DEPRECATED(5.5,"Function will be removed from UE5.7")
+	UFUNCTION(BlueprintPure, Category = "Virtual Production", meta=(DeprecatedFunction, DeprecationMessage="This function will be removed from UE5.7"))
 	UUserWidget* GetPanelWidget(const FName& PanelID) const;
 
-	UFUNCTION(BlueprintCallable, Category = "Virtual Production")
+	UE_DEPRECATED(5.5,"Function will be removed from UE5.7")
+	UFUNCTION(BlueprintCallable, Category = "Virtual Production", meta=(DeprecatedFunction, DeprecationMessage="This function will be removed from UE5.7"))
 	static TArray<UVREditorInteractor*> GetActiveEditorVRControllers();
 
-	UFUNCTION(BlueprintPure, Category = "Virtual Production")
+	UE_DEPRECATED(5.5,"Function will be removed from UE5.7")
+	UFUNCTION(BlueprintPure, Category = "Virtual Production", meta=(DeprecatedFunction, DeprecationMessage="This function will be removed from UE5.7"))
 	static const FName GetVProdPanelID(const EVProdPanelIDs Panel)
 	{
 		switch (Panel)
@@ -148,94 +167,118 @@ public:
 		return VProdPanelID;
 	};
 
-	UFUNCTION(BlueprintPure, Category = "Virtual Production")
+	UE_DEPRECATED(5.5,"Function will be removed from UE5.7")
+	UFUNCTION(BlueprintPure, Category = "Virtual Production", meta=(DeprecatedFunction, DeprecationMessage="This function will be removed from UE5.7"))
 	static FString GetDirectorName();
 
-	UFUNCTION(BlueprintPure, Category = "Virtual Production")
+	UE_DEPRECATED(5.5,"Function will be removed from UE5.7")
+	UFUNCTION(BlueprintPure, Category = "Virtual Production", meta=(DeprecatedFunction, DeprecationMessage="This function will be removed from UE5.7"))
 	static FString GetShowName();
 
+	
 	/** Whether the VR user wants to use the metric system instead of imperial */
-	UFUNCTION(BlueprintPure, Category = "Virtual Production")
+	UE_DEPRECATED(5.5,"Function will be removed from UE5.7")
+	UFUNCTION(BlueprintPure, Category = "Virtual Production", meta=(DeprecatedFunction, DeprecationMessage="This function will be removed from UE5.7"))
 	static bool IsUsingMetricSystem();
 
 	/** Set whether the VR user wants to use the metric system instead of imperial */
-	UFUNCTION(BlueprintCallable, Category = "Virtual Production")
+	UE_DEPRECATED(5.5,"Function will be removed from UE5.7")
+	UFUNCTION(BlueprintCallable, Category = "Virtual Production", meta=(DeprecatedFunction, DeprecationMessage="This function will be removed from UE5.7"))
 	static void SetIsUsingMetricSystem(const bool bInUseMetricSystem);
 	
 	/** Whether the VR user wants to have the transform gizmo enabled */
-	UFUNCTION(BlueprintPure, Category = "Virtual Production")
+	UE_DEPRECATED(5.5,"Function will be removed from UE5.7")
+	UFUNCTION(BlueprintPure, Category = "Virtual Production", meta=(DeprecatedFunction, DeprecationMessage="This function will be removed from UE5.7"))
 	static bool IsUsingTransformGizmo();
 
 	/** Set whether the VR user wants to have the transform gizmo enabled */
-	UFUNCTION(BlueprintCallable, Category = "Virtual Production")
+	UE_DEPRECATED(5.5,"Function will be removed from UE5.7")
+	UFUNCTION(BlueprintCallable, Category = "Virtual Production", meta=(DeprecatedFunction, DeprecationMessage="This function will be removed from UE5.7"))
 	static void SetIsUsingTransformGizmo(const bool bInIsUsingTransformGizmo);
 
 	/** Set value of cvar "VI.ShowTransformGizmo" */
-	UFUNCTION(BlueprintCallable, Category = "Virtual Production")
+	UE_DEPRECATED(5.5,"Function will be removed from UE5.7")
+	UFUNCTION(BlueprintCallable, Category = "Virtual Production", meta=(DeprecatedFunction, DeprecationMessage="This function will be removed from UE5.7"))
 	static void SetShowTransformGizmoCVar(const bool bInShowTransformGizmoCVar);
 
 	/** Get flight speed for scouting in VR */
-	UFUNCTION(BlueprintPure, Category = "Virtual Production")
+	UE_DEPRECATED(5.5,"Function will be removed from UE5.7")
+	UFUNCTION(BlueprintPure, Category = "Virtual Production", meta=(DeprecatedFunction, DeprecationMessage="This function will be removed from UE5.7"))
 	static float GetFlightSpeed();
 
 	/** Set flight speed for scouting in VR */
-	UFUNCTION(BlueprintCallable, Category = "Virtual Production")
+	UE_DEPRECATED(5.5,"Function will be removed from UE5.7")
+	UFUNCTION(BlueprintCallable, Category = "Virtual Production", meta=(DeprecatedFunction, DeprecationMessage="This function will be removed from UE5.7"))
 	static void SetFlightSpeed(const float InFlightSpeed);
 
 	/** Get grip nav speed for scouting in VR */
-	UFUNCTION(BlueprintPure, Category = "Virtual Production")
+	UE_DEPRECATED(5.5,"Function will be removed from UE5.7")
+	UFUNCTION(BlueprintPure, Category = "Virtual Production", meta=(DeprecatedFunction, DeprecationMessage="This function will be removed from UE5.7"))
 	static float GetGripNavSpeed();
 
 	/** Set grip nav speed for scouting in VR */
-	UFUNCTION(BlueprintCallable, Category = "Virtual Production")
+	UE_DEPRECATED(5.5,"Function will be removed from UE5.7")
+	UFUNCTION(BlueprintCallable, Category = "Virtual Production", meta=(DeprecatedFunction, DeprecationMessage="This function will be removed from UE5.7"))
 	static void SetGripNavSpeed(const float InGripNavSpeed);
 
 	/** Whether grip nav inertia is enabled when scouting in VR */
-	UFUNCTION(BlueprintPure, Category = "Virtual Production")
+	UE_DEPRECATED(5.5,"Function will be removed from UE5.7")
+	UFUNCTION(BlueprintPure, Category = "Virtual Production", meta=(DeprecatedFunction, DeprecationMessage="This function will be removed from UE5.7"))
 	static bool IsUsingInertiaDamping();
 
 	/** Set whether grip nav inertia is enabled when scouting in VR */
-	UFUNCTION(BlueprintCallable, Category = "Virtual Production")
+	UE_DEPRECATED(5.5,"Function will be removed from UE5.7")
+	UFUNCTION(BlueprintCallable, Category = "Virtual Production", meta=(DeprecatedFunction, DeprecationMessage="This function will be removed from UE5.7"))
 	static void SetIsUsingInertiaDamping(const bool bInIsUsingInertiaDamping);
 
 	/** Set value of cvar "VI.HighSpeedInertiaDamping" */
-	UFUNCTION(BlueprintCallable, Category = "Virtual Production")
+	UE_DEPRECATED(5.5,"Function will be removed from UE5.7")
+	UFUNCTION(BlueprintCallable, Category = "Virtual Production", meta=(DeprecatedFunction, DeprecationMessage="This function will be removed from UE5.7"))
 	static void SetInertiaDampingCVar(const float InInertiaDamping);
 
 	/** Whether the helper system on the controllers is enabled */
-	UFUNCTION(BlueprintPure, Category = "Virtual Production")
+	UE_DEPRECATED(5.5,"Function will be removed from UE5.7")
+	UFUNCTION(BlueprintPure, Category = "Virtual Production", meta=(DeprecatedFunction, DeprecationMessage="This function will be removed from UE5.7"))
 	static bool IsHelperSystemEnabled();
 
 	/** Set whether the helper system on the controllers is enabled   */
-	UFUNCTION(BlueprintCallable, Category = "Virtual Production")
+	UE_DEPRECATED(5.5,"Function will be removed from UE5.7")
+	UFUNCTION(BlueprintCallable, Category = "Virtual Production", meta=(DeprecatedFunction, DeprecationMessage="This function will be removed from UE5.7"))
 	static void SetIsHelperSystemEnabled(const bool bInIsHelperSystemEnabled);
 
 	/** Get VR Editor Mode object */
-	UFUNCTION(BlueprintPure, Category = "Virtual Production")
+	UE_DEPRECATED(5.5,"Function will be removed from UE5.7")
+	UFUNCTION(BlueprintPure, Category = "Virtual Production", meta=(DeprecatedFunction, DeprecationMessage="This function will be removed from UE5.7"))
 	static UVREditorMode* GetVREditorMode();
 
 	/** Enter VR Mode */
-	UFUNCTION(BlueprintCallable, Category = "Virtual Production")
+	UE_DEPRECATED(5.5,"Function will be removed from UE5.7")
+	UFUNCTION(BlueprintCallable, Category = "Virtual Production", meta=(DeprecatedFunction, DeprecationMessage="This function will be removed from UE5.7"))
 	static bool EnterVRMode();
 
 	/** Exit VR Mode */
-	UFUNCTION(BlueprintCallable, Category = "Virtual Production")
+	UE_DEPRECATED(5.5,"Function will be removed from UE5.7")
+	UFUNCTION(BlueprintCallable, Category = "Virtual Production", meta=(DeprecatedFunction, DeprecationMessage="This function will be removed from UE5.7"))
 	static void ExitVRMode();
 
 	/** Whether location grid snapping is enabled */
-	UFUNCTION(BlueprintCallable, Category = "Virtual Production")
+	UE_DEPRECATED(5.5,"Function will be removed from UE5.7")
+	UFUNCTION(BlueprintCallable, Category = "Virtual Production", meta=(DeprecatedFunction, DeprecationMessage="This function will be removed from UE5.7"))
 	static bool IsLocationGridSnappingEnabled();
 
 	/** Toggle location grid snapping */
-	UFUNCTION(BlueprintCallable, Category = "Virtual Production")
+	UE_DEPRECATED(5.5,"Function will be removed from UE5.7")
+	UFUNCTION(BlueprintCallable, Category = "Virtual Production", meta=(DeprecatedFunction, DeprecationMessage="This function will be removed from UE5.7"))
 	static void ToggleLocationGridSnapping();
 
 	/** Whether rotation grid snapping is enabled */
-	UFUNCTION(BlueprintCallable, Category = "Virtual Production")
+	UE_DEPRECATED(5.5,"Function will be removed from UE5.7")
+	UFUNCTION(BlueprintCallable, Category = "Virtual Production", meta=(DeprecatedFunction, DeprecationMessage="This function will be removed from UE5.7"))
 	static bool IsRotationGridSnappingEnabled();
 
 	/** Toggle rotation grid snapping */
-	UFUNCTION(BlueprintCallable, Category = "Virtual Production")
+	UE_DEPRECATED(5.5,"Function will be removed from UE5.7")
+	UFUNCTION(BlueprintCallable, Category = "Virtual Production", meta=(DeprecatedFunction, DeprecationMessage="This function will be removed from UE5.7"))
 	static void ToggleRotationGridSnapping();
 
 private:
@@ -253,3 +296,5 @@ private:
 	static const FName VProdPanelMeasureID;
 	static const FName VProdPanelGafferID;
 };
+
+PRAGMA_ENABLE_DEPRECATION_WARNINGS
\ No newline at end of file
diff --git a/Engine/Plugins/Experimental/VirtualProductionUtilities/Source/VPUtilitiesEditor/Public/VPUtilitiesEditorSettings.h b/Engine/Plugins/Experimental/VirtualProductionUtilities/Source/VPUtilitiesEditor/Public/VPUtilitiesEditorSettings.h
--- a/Engine/Plugins/Experimental/VirtualProductionUtilities/Source/VPUtilitiesEditor/Public/VPUtilitiesEditorSettings.h
+++ b/Engine/Plugins/Experimental/VirtualProductionUtilities/Source/VPUtilitiesEditor/Public/VPUtilitiesEditorSettings.h
@@ -22,35 +22,43 @@ public:
 	UVPUtilitiesEditorSettings();
 
 	/** The default user interface that we'll use for virtual scouting */
-	UPROPERTY(EditAnywhere, config, Category = "Legacy Virtual Scouting", meta = (DisplayName = "Virtual Scouting User Interface"))
+	UE_DEPRECATED(5.5,"This property  will be removed from UE5.7")
+	UPROPERTY(EditAnywhere, config, Category = "Legacy Virtual Scouting", meta = (DisplayName = "Virtual Scouting User Interface", DeprecatedProperty, DeprecationMessage="This property is deprecated"))
 	TSoftClassPtr<UEditorUtilityWidget> VirtualScoutingUI;
 	
 	/** Speed when flying in VR*/
-	UPROPERTY(EditAnywhere, config, Category = "Legacy Virtual Scouting", meta = (ClampMin = "0.0", ClampMax = "1.0", DisplayName = "Virtual Scouting Flight Speed"))
+	UE_DEPRECATED(5.5,"This property  will be removed from UE5.7")
+	UPROPERTY(EditAnywhere, config, Category = "Legacy Virtual Scouting", meta = (ClampMin = "0.0", ClampMax = "1.0", DisplayName = "Virtual Scouting Flight Speed", DeprecatedProperty, DeprecationMessage="This property is deprecated"))
 	float FlightSpeed = 0.5f;
 	
 	/** Speed when using grip nav in VR */
-	UPROPERTY(EditAnywhere, config, Category = "Legacy Virtual Scouting", meta = (ClampMin = "0.0", ClampMax = "1.0", DisplayName = "Virtual Scouting Grip Nav Speed"))
+	UE_DEPRECATED(5.5,"This property  will be removed from UE5.7")
+	UPROPERTY(EditAnywhere, config, Category = "Legacy Virtual Scouting", meta = (ClampMin = "0.0", ClampMax = "1.0", DisplayName = "Virtual Scouting Grip Nav Speed", DeprecatedProperty, DeprecationMessage="This property is deprecated"))
 	float GripNavSpeed = 0.25f;
 	
 	/** Whether to use the metric system or imperial for measurements */
-	UPROPERTY(EditAnywhere, config, Category = "Legacy Virtual Scouting", meta = (DisplayName = "Show Measurements In Metric Units"))
+	UE_DEPRECATED(5.5,"This property  will be removed from UE5.7")
+	UPROPERTY(EditAnywhere, config, Category = "Legacy Virtual Scouting", meta = (DisplayName = "Show Measurements In Metric Units", DeprecatedProperty, DeprecationMessage="This property is deprecated"))
 	bool bUseMetric = false;
 	
 	/** Whether to enable or disable the transform gizmo */
-	UPROPERTY(EditAnywhere, config, Category = "Legacy Virtual Scouting", meta = (DisplayName = "Enable Transform Gizmo In VR"))
+	UE_DEPRECATED(5.5,"This property  will be removed from UE5.7")
+	UPROPERTY(EditAnywhere, config, Category = "Legacy Virtual Scouting", meta = (DisplayName = "Enable Transform Gizmo In VR", DeprecatedProperty, DeprecationMessage="This property is deprecated"))
 	bool bUseTransformGizmo = false;
 	
 	/** If true, the user will use inertia damping to stop after grip nav. Otherwise the user will just stop immediately */
-	UPROPERTY(EditAnywhere, config, Category = "Legacy Virtual Scouting", meta = (DisplayName = "Use Grip Inertia Damping"))
+	UE_DEPRECATED(5.5,"This property  will be removed from UE5.7")
+	UPROPERTY(EditAnywhere, config, Category = "Legacy Virtual Scouting", meta = (DisplayName = "Use Grip Inertia Damping", DeprecatedProperty, DeprecationMessage="This property is deprecated"))
 	bool bUseGripInertiaDamping = true;
 	
 	/** Damping applied to inertia */
-	UPROPERTY(EditAnywhere, config, Category = "Legacy Virtual Scouting", meta = (ClampMin = "0.0", ClampMax = "1.0", DisplayName = "Inertia Damping"))
+	UE_DEPRECATED(5.5,"This property  will be removed from UE5.7")
+	UPROPERTY(EditAnywhere, config, Category = "Legacy Virtual Scouting", meta = (ClampMin = "0.0", ClampMax = "1.0", DisplayName = "Inertia Damping", DeprecatedProperty, DeprecationMessage="This property is deprecated"))
 	float InertiaDamping = 0.95f;
 
 	/** Whether the helper system on the controllers is enabled */
-	UPROPERTY(EditAnywhere, config, Category = "Legacy Virtual Scouting", meta = (DisplayName = "Helper System Enabled "))
+	UE_DEPRECATED(5.5,"Code will be removed from UE5.7")
+	UPROPERTY(EditAnywhere, config, Category = "Legacy Virtual Scouting", meta = (DisplayName = "Helper System Enabled ", DeprecatedProperty, DeprecationMessage="This property is deprecated"))
 	bool bIsHelperSystemEnabled = true;
 	
 	/** When enabled, an OSC server will automatically start on launch. */
@@ -70,17 +78,21 @@ public:
 	TArray<FSoftObjectPath> StartupOSCListeners;
 
 	/** ScoutingSubsystem class to use for Blueprint helpers */
+	UE_DEPRECATED(5.5,"Code will be removed from UE5.7")
 	UPROPERTY(config, meta = (MetaClass = "/Script/VPUtilitiesEditor.VPScoutingSubsystemHelpersBase"))
 	FSoftClassPath ScoutingSubsystemEditorUtilityClassPath;
 
 	/** GestureManager class to use by the ScoutingSubsystem */
+	UE_DEPRECATED(5.5,"Code will be removed from UE5.7")
 	UPROPERTY(config, meta = (MetaClass = "/Script/VPUtilitiesEditor.VPScoutingSubsystemGestureManagerBase"))
 	FSoftClassPath GestureManagerEditorUtilityClassPath;
 
 	/** GestureManager class to use by the ScoutingSubsystem */
+	UE_DEPRECATED(5.5,"Code will be removed from UE5.7")
 	UPROPERTY(config)
 	TArray<FSoftClassPath> AdditionnalClassToLoad;
 
+	UE_DEPRECATED(5.5,"Code will be removed from UE5.7")
 	UPROPERTY(config, meta = (AllowedClasses = "/Script/Engine.StaticMesh"))
 	FSoftObjectPath VPSplinePreviewMeshPath;
 };
\ No newline at end of file
diff --git a/Engine/Plugins/Experimental/VirtualProductionUtilities/Source/VPUtilitiesEditor/Public/VRTool.h b/Engine/Plugins/Experimental/VirtualProductionUtilities/Source/VPUtilitiesEditor/Public/VRTool.h
--- a/Engine/Plugins/Experimental/VirtualProductionUtilities/Source/VPUtilitiesEditor/Public/VRTool.h
+++ b/Engine/Plugins/Experimental/VirtualProductionUtilities/Source/VPUtilitiesEditor/Public/VRTool.h
@@ -9,8 +9,9 @@
 /**
  * Base class for virtual production VR tools
  */
+
 UCLASS()
-class AVRTool : public AVPTransientEditorTickableActorBase
+class UE_DEPRECATED(5.5,"This class will be removed from UE5.7") AVRTool : public AVPTransientEditorTickableActorBase
 {
 	GENERATED_BODY()
 
@@ -27,10 +28,10 @@ class AVRTool : public AVPTransientEditorTickableActorBase
 	{
 		return EditorOnlyInputComponent.Get();
 	}
-
+	
 	UFUNCTION(BlueprintSetter, Category = "Input|Editor")
 	void SetReceivesEditorInput(bool bInValue);
-
+	
 	UFUNCTION(BlueprintGetter, BlueprintPure, Category = "Input|Editor")
 	bool GetReceivesEditorInput() const
 	{
@@ -41,15 +42,15 @@ private:
 
 	/** Creates the EditorOnlyInputComponent if it does not already exist and registers all subobject callbacks to it */
 	void CreateEditorInput();
-
+	
 	/** Removes the EditorOnlyInputComponent from this utility actor */
 	void RemoveEditorInput();
-
-	UPROPERTY(Transient, DuplicateTransient)
+	
+	UPROPERTY(Transient, DuplicateTransient, meta=(DeprecatedProperty, DeprecationMessage="This is deprecated"))
 	TObjectPtr<UInputComponent> EditorOnlyInputComponent;
 
 	/** If set to true, then this actor will be able to recieve input delegate callbacks when in the editor. */
-	UPROPERTY(EditAnywhere, Category = "Input|Editor", BlueprintSetter = SetReceivesEditorInput, BlueprintGetter = GetReceivesEditorInput)
+	UPROPERTY(EditAnywhere, Category = "Input|Editor", BlueprintSetter = SetReceivesEditorInput, BlueprintGetter = GetReceivesEditorInput, meta=(DeprecatedProperty, DeprecationMessage="This is deprecated"))
 	bool bReceivesEditorInput = false;
 
 };
\ No newline at end of file
diff --git a/Engine/Plugins/VirtualProduction/VirtualScouting/Source/VirtualScoutingEditor/Private/VPScoutingMode.cpp b/Engine/Plugins/VirtualProduction/VirtualScouting/Source/VirtualScoutingEditor/Private/VPScoutingMode.cpp
deleted file mode 100644
diff --git a/Engine/Plugins/VirtualProduction/VirtualScouting/Source/VirtualScoutingEditor/Private/VPScoutingMode.h b/Engine/Plugins/VirtualProduction/VirtualScouting/Source/VirtualScoutingEditor/Private/VPScoutingMode.h
deleted file mode 100644
