[Virtual Produciton Utilities]

Fix for cook error fallout from removing deprecated content. Corrected by using the base ini to point to the native classes of the deleted assets so async load calls in the VPScouting subsystem to load those classes don't fail. Made abstract Gesture Manager class "NotPlaceable" instead of "Abstract" to enable packaging. Practically this change is just to unblock packaging, not a user-facing change. Assumes they won't be using this class in their work.

#rnx

diff --git a/Engine/Plugins/Experimental/VirtualProductionUtilities/Config/BaseVirtualProductionUtilities.ini b/Engine/Plugins/Experimental/VirtualProductionUtilities/Config/BaseVirtualProductionUtilities.ini
--- a/Engine/Plugins/Experimental/VirtualProductionUtilities/Config/BaseVirtualProductionUtilities.ini
+++ b/Engine/Plugins/Experimental/VirtualProductionUtilities/Config/BaseVirtualProductionUtilities.ini
@@ -27,3 +27,7 @@ BookmarkMaterialPath=/VirtualProductionUtilities/Bookmark/MI_BookmarkVPNoUser.MI
 BookmarkSplineMeshPath=/Engine/VREditor/FloatingText/LineSegmentCylinder.LineSegmentCylinder
 BookmarkSplineMeshMaterialPath=/Engine/VREditor/LaserPointer/LaserPointerMaterialInst.LaserPointerMaterialInst
 BookmarkLabelMaterialPath=/VirtualProductionUtilities/Bookmark/CameraFacingTextMaterial.CameraFacingTextMaterial
+
+[/Script/VPUtilitiesEditor.VPUtilitiesEditorSettings]
+ScoutingSubsystemEditorUtilityClassPath=/Script/VPUtilitiesEditor.VPScoutingSubsystemHelpersBase
+GestureManagerEditorUtilityClassPath=/Script/VPUtilitiesEditor.VPScoutingSubsystemGestureManagerBase
diff --git a/Engine/Plugins/Experimental/VirtualProductionUtilities/Source/VPUtilitiesEditor/Public/VPScoutingSubsystem.h b/Engine/Plugins/Experimental/VirtualProductionUtilities/Source/VPUtilitiesEditor/Public/VPScoutingSubsystem.h
--- a/Engine/Plugins/Experimental/VirtualProductionUtilities/Source/VPUtilitiesEditor/Public/VPScoutingSubsystem.h
+++ b/Engine/Plugins/Experimental/VirtualProductionUtilities/Source/VPUtilitiesEditor/Public/VPScoutingSubsystem.h
@@ -39,7 +39,7 @@ class UE_DEPRECATED(5.5,"Code will be removed from UE5.7") UVPScoutingSubsystemH
  * Base class of the gesture manager defined in BP
  */
 
-UCLASS(Abstract, Blueprintable, MinimalAPI, meta = (ShowWorldContextPin))
+UCLASS(NotPlaceable, Blueprintable, MinimalAPI, meta = (ShowWorldContextPin))
 class UE_DEPRECATED(5.5,"This class will be removed from UE5.7") UVPScoutingSubsystemGestureManagerBase : public UObject, public FTickableEditorObject
 {
 	GENERATED_BODY()
