Fix horde issue: 678419 Removed warnings about a deprecated UClass.
#rnx

diff --git a/Engine/Plugins/Experimental/VirtualProductionUtilities/Source/VPUtilitiesEditor/Private/VPScoutingSubsystem.cpp b/Engine/Plugins/Experimental/VirtualProductionUtilities/Source/VPUtilitiesEditor/Private/VPScoutingSubsystem.cpp
--- a/Engine/Plugins/Experimental/VirtualProductionUtilities/Source/VPUtilitiesEditor/Private/VPScoutingSubsystem.cpp
+++ b/Engine/Plugins/Experimental/VirtualProductionUtilities/Source/VPUtilitiesEditor/Private/VPScoutingSubsystem.cpp
@@ -123,10 +123,6 @@ void UVPScoutingSubsystem::Initialize(FSubsystemCollectionBase& Collection)
 		{
 			VPSubsystemHelpers = NewObject<UVPScoutingSubsystemHelpersBase>(GetTransientPackage(), EditorUtilityClass);
 		}
-		else
-		{
-			UE_LOG(LogVPUtilitiesEditor, Warning, TEXT("Failed loading VPScoutingHelpers \"%s\""), *GetDefault<UVPUtilitiesEditorSettings>()->ScoutingSubsystemEditorUtilityClassPath.ToString());
-		}
 	}));
 	
 	// to do final initializations at the right time
@@ -152,10 +148,6 @@ void UVPScoutingSubsystem::OnEngineInitComplete()
 	{
 		GestureManager = NewObject<UVPScoutingSubsystemGestureManagerBase>(GetTransientPackage(), EditorUtilityClass);
 	}
-	else
-	{
-		UE_LOG(LogVPUtilitiesEditor, Warning, TEXT("Failed loading VPScoutingHelpers \"%s\""), *GetDefault<UVPUtilitiesEditorSettings>()->GestureManagerEditorUtilityClassPath.ToString());
-	}
 
 	// In debug some asset take a long time to load and crash the engine, preload those asset in async mode to prevent that
 	for (const FSoftClassPath& ClassAssetPath : GetDefault<UVPUtilitiesEditorSettings>()->AdditionnalClassToLoad)
