VirtualProductionUtilities: Deprecate BP_VPBookmarkActor

#virtualized

diff --git a/Engine/Plugins/Experimental/VirtualProductionUtilities/Config/BaseVirtualProductionUtilities.ini b/Engine/Plugins/Experimental/VirtualProductionUtilities/Config/BaseVirtualProductionUtilities.ini
--- a/Engine/Plugins/Experimental/VirtualProductionUtilities/Config/BaseVirtualProductionUtilities.ini
+++ b/Engine/Plugins/Experimental/VirtualProductionUtilities/Config/BaseVirtualProductionUtilities.ini
@@ -20,6 +20,9 @@
 +PropertyRedirects=(OldName="VPFullScreenUserWidget.PostProcessDisplayType",NewName="VPFullScreenUserWidget.PostProcessDisplayTypeWithBlendMaterial")
 +EnumRedirects=(OldName="/Script/VPUtilities.EVPWidgetDisplayType", NewName="/Script/VPUtilities.EVPWidgetDisplayType", ValueChanges=(("PostProcess", "PostProcessWithBlendMaterial")))
 
+; 5.5 > 5.6
++PackageRedirects=(OldName="/VirtualProductionUtilities/Bookmark/BP_VPBookmarkActor",NewName="/VirtualProductionUtilities/DEPRECATED/BP_VPBookmarkActor")
+
 [/Script/VPUtilities.VPBookmarkSettings]
 
 BookmarkMeshPath=/VirtualProductionUtilities/Bookmark/BookmarkVPMesh.BookmarkVPMesh
