Fix CIS issue 1024963 compile warning using a deprecated property.

diff --git a/Engine/Plugins/Experimental/VirtualProductionUtilities/Source/VPUtilities/Public/Widgets/Fullscreen/VPFullScreenUserWidget_PostProcess.h b/Engine/Plugins/Experimental/VirtualProductionUtilities/Source/VPUtilities/Public/Widgets/Fullscreen/VPFullScreenUserWidget_PostProcess.h
--- a/Engine/Plugins/Experimental/VirtualProductionUtilities/Source/VPUtilities/Public/Widgets/Fullscreen/VPFullScreenUserWidget_PostProcess.h
+++ b/Engine/Plugins/Experimental/VirtualProductionUtilities/Source/VPUtilities/Public/Widgets/Fullscreen/VPFullScreenUserWidget_PostProcess.h
@@ -16,13 +16,15 @@ struct FVPFullScreenUserWidget_PostProcess : public FVPFullScreenUserWidget_Post
 {
 	GENERATED_BODY()
 
+	PRAGMA_DISABLE_DEPRECATION_WARNINGS
 #if WITH_EDITORONLY_DATA
 	/** List of composure layers that are expecting to use the WidgetRenderTarget. */
 	UE_DEPRECATED(5.7, "Composure layers are no longer supported on the VP Full screen user widget.")		
 	UPROPERTY()
 	TArray<TObjectPtr<AActor>> ComposureLayerTargets;
 #endif
-	
+	PRAGMA_ENABLE_DEPRECATION_WARNINGS
+
 	FVPFullScreenUserWidget_PostProcess();
 
 	void SetCustomPostProcessSettingsSource(TWeakObjectPtr<UObject> InCustomPostProcessSettingsSource);
