ViewportInteraction: Deprecate ViewportInteraction module alongside VR Editor.

diff --git a/Engine/Plugins/Experimental/VirtualProductionUtilities/Source/VPBookmarkEditor/Private/VPBookmarkTypeActions.cpp b/Engine/Plugins/Experimental/VirtualProductionUtilities/Source/VPBookmarkEditor/Private/VPBookmarkTypeActions.cpp
--- a/Engine/Plugins/Experimental/VirtualProductionUtilities/Source/VPBookmarkEditor/Private/VPBookmarkTypeActions.cpp
+++ b/Engine/Plugins/Experimental/VirtualProductionUtilities/Source/VPBookmarkEditor/Private/VPBookmarkTypeActions.cpp
@@ -20,6 +20,7 @@
 
 namespace VPBookmark
 {
+PRAGMA_DISABLE_DEPRECATION_WARNINGS
 	static UViewportWorldInteraction* GetViewportWorldInteraction()
 	{
 		if (UEditorWorldExtensionManager* ExtensionManager = GEditor->GetEditorWorldExtensionsManager())
@@ -32,6 +33,7 @@ namespace VPBookmark
 		}
 		return nullptr;
 	}
+PRAGMA_ENABLE_DEPRECATION_WARNINGS
 
 
 	static FTransform GetEditorViewportTransform(FEditorViewportClient& InViewportClient)
diff --git a/Engine/Plugins/Experimental/VirtualProductionUtilities/Source/VPUtilities/Private/Libraries/VPBlueprintLibrary.cpp b/Engine/Plugins/Experimental/VirtualProductionUtilities/Source/VPUtilities/Private/Libraries/VPBlueprintLibrary.cpp
--- a/Engine/Plugins/Experimental/VirtualProductionUtilities/Source/VPUtilities/Private/Libraries/VPBlueprintLibrary.cpp
+++ b/Engine/Plugins/Experimental/VirtualProductionUtilities/Source/VPUtilities/Private/Libraries/VPBlueprintLibrary.cpp
@@ -32,6 +32,7 @@ namespace VPBlueprintLibrary
 			nullptr;
 	}
 
+PRAGMA_DISABLE_DEPRECATION_WARNINGS
 	UViewportWorldInteraction* GetViewportWorldInteraction(FString& FailMessage)
 	{
 		if (FLevelEditorViewportClient* Client = GetViewPortClient())
@@ -56,6 +57,7 @@ namespace VPBlueprintLibrary
 		UE_LOG(LogVPUtilities, Warning, TEXT("UVPBlueprintLibrary::GetViewportWorldInteraction - Failed to get VPI. %s"), *FailMessage);
 		return nullptr;
 	}
+PRAGMA_ENABLE_DEPRECATION_WARNINGS
 #endif
 }
 
@@ -144,6 +146,7 @@ FTransform UVPBlueprintLibrary::GetEditorViewportTransform()
 
 FTransform UVPBlueprintLibrary::GetEditorVRHeadTransform()
 {
+PRAGMA_DISABLE_DEPRECATION_WARNINGS
 #if WITH_EDITOR
 	FString ErrorText(TEXT("Head Transform will be invalid."));
 	if (const UViewportWorldInteraction* VPI = VPBlueprintLibrary::GetViewportWorldInteraction(ErrorText))
@@ -151,6 +154,7 @@ FTransform UVPBlueprintLibrary::GetEditorVRHeadTransform()
 		return VPI->GetHeadTransform();
 	}
 #endif
+PRAGMA_ENABLE_DEPRECATION_WARNINGS
 
 	return FTransform::Identity;
 }
@@ -158,6 +162,7 @@ FTransform UVPBlueprintLibrary::GetEditorVRHeadTransform()
 
 FTransform UVPBlueprintLibrary::GetEditorVRRoomTransform()
 {
+PRAGMA_DISABLE_DEPRECATION_WARNINGS
 #if WITH_EDITOR
 	FString ErrorText(TEXT("Room Transform will be invalid."));
 	if (const UViewportWorldInteraction* VPI = VPBlueprintLibrary::GetViewportWorldInteraction(ErrorText))
@@ -165,6 +170,7 @@ FTransform UVPBlueprintLibrary::GetEditorVRRoomTransform()
 		return VPI->GetRoomTransform();
 	}
 #endif
+PRAGMA_ENABLE_DEPRECATION_WARNINGS
 
 	return FTransform::Identity;
 }
@@ -195,6 +201,7 @@ PRAGMA_ENABLE_DEPRECATION_WARNINGS
 FVector UVPBlueprintLibrary::GetVREditorLaserHoverLocation()
 {
 #if WITH_EDITOR
+PRAGMA_DISABLE_DEPRECATION_WARNINGS
 	FString ErrorText(TEXT("VR laser hit location will be invalid."));
 	if (const UViewportWorldInteraction* VPI = VPBlueprintLibrary::GetViewportWorldInteraction(ErrorText))
 	{
@@ -202,7 +209,6 @@ FVector UVPBlueprintLibrary::GetVREditorLaserHoverLocation()
 
 		for (UViewportInteractor* Interactor : Interactors)
 		{
-PRAGMA_DISABLE_DEPRECATION_WARNINGS
 			if (UVREditorInteractor* EdInteractor = Cast<UVREditorInteractor>(Interactor))
 			{
 				if (EdInteractor->GetControllerType() == EControllerType::Laser)
@@ -213,9 +219,9 @@ PRAGMA_DISABLE_DEPRECATION_WARNINGS
 					return EdInteractor->GetInteractorData().LastHoverLocationOverUI;
 				}
 			}
-PRAGMA_ENABLE_DEPRECATION_WARNINGS
 		}
 	}
+PRAGMA_ENABLE_DEPRECATION_WARNINGS
 #endif
 
 	return FVector(0);
@@ -224,6 +230,7 @@ PRAGMA_ENABLE_DEPRECATION_WARNINGS
 FVector UVPBlueprintLibrary::GetLeftInteractorLocation()
 {
 #if WITH_EDITOR
+PRAGMA_DISABLE_DEPRECATION_WARNINGS
 	FString ErrorText(TEXT("VR interactor location will be invalid."));
 	if (const UViewportWorldInteraction* VPI = VPBlueprintLibrary::GetViewportWorldInteraction(ErrorText))
 	{
@@ -231,7 +238,6 @@ FVector UVPBlueprintLibrary::GetLeftInteractorLocation()
 
 		for (UViewportInteractor* Interactor : Interactors)
 		{
-PRAGMA_DISABLE_DEPRECATION_WARNINGS
 			if (UVREditorInteractor* EdInteractor = Cast<UVREditorInteractor>(Interactor))
 			{
 				if (EdInteractor->GetControllerType() == EControllerType::Laser)
@@ -242,9 +248,9 @@ PRAGMA_DISABLE_DEPRECATION_WARNINGS
 					return InteractorLocation;
 				}
 			}
-PRAGMA_ENABLE_DEPRECATION_WARNINGS
 		}
 	}
+PRAGMA_ENABLE_DEPRECATION_WARNINGS
 #endif
 
 	return FVector(0);
@@ -253,12 +259,14 @@ PRAGMA_ENABLE_DEPRECATION_WARNINGS
 bool UVPBlueprintLibrary::EditorUndo()
 {
 #if WITH_EDITOR
+PRAGMA_DISABLE_DEPRECATION_WARNINGS
 	FString ErrorText(TEXT("Undo did not execute."));
 	if (UViewportWorldInteraction* VPI = VPBlueprintLibrary::GetViewportWorldInteraction(ErrorText))
 	{
 		VPI->Undo();
 		return true;
 	}
+PRAGMA_ENABLE_DEPRECATION_WARNINGS
 #endif
 
 	return false;
@@ -268,12 +276,14 @@ bool UVPBlueprintLibrary::EditorUndo()
 bool UVPBlueprintLibrary::EditorRedo()
 {
 #if WITH_EDITOR
+PRAGMA_DISABLE_DEPRECATION_WARNINGS
 	FString ErrorText(TEXT("Redo did not execute."));
 	if (UViewportWorldInteraction* VPI = VPBlueprintLibrary::GetViewportWorldInteraction(ErrorText))
 	{
 		VPI->Redo();
 		return true;
 	}
+PRAGMA_ENABLE_DEPRECATION_WARNINGS
 #endif
 
 	return false;
@@ -283,12 +293,14 @@ bool UVPBlueprintLibrary::EditorRedo()
 bool UVPBlueprintLibrary::EditorDeleteSelectedObjects()
 {
 #if WITH_EDITOR
+PRAGMA_DISABLE_DEPRECATION_WARNINGS
 	FString ErrorText(TEXT("Delete did not execute."));
 	if (UViewportWorldInteraction* VPI = VPBlueprintLibrary::GetViewportWorldInteraction(ErrorText))
 	{
 		VPI->DeleteSelectedObjects();
 		return true;
 	}
+PRAGMA_ENABLE_DEPRECATION_WARNINGS
 #endif
 
 	return false;
@@ -297,12 +309,14 @@ bool UVPBlueprintLibrary::EditorDeleteSelectedObjects()
 bool UVPBlueprintLibrary::EditorDuplicate()
 {
 #if WITH_EDITOR
+PRAGMA_DISABLE_DEPRECATION_WARNINGS
 	FString ErrorText(TEXT("Duplicate did not execute."));
 	if (UViewportWorldInteraction* VPI = VPBlueprintLibrary::GetViewportWorldInteraction(ErrorText))
 	{
 		VPI->Duplicate();
 		return true;
 	}
+PRAGMA_ENABLE_DEPRECATION_WARNINGS
 #endif
 
 	return false;
diff --git a/Engine/Plugins/Experimental/VirtualProductionUtilities/Source/VPUtilities/Public/Libraries/VPBlueprintLibrary.h b/Engine/Plugins/Experimental/VirtualProductionUtilities/Source/VPUtilities/Public/Libraries/VPBlueprintLibrary.h
--- a/Engine/Plugins/Experimental/VirtualProductionUtilities/Source/VPUtilities/Public/Libraries/VPBlueprintLibrary.h
+++ b/Engine/Plugins/Experimental/VirtualProductionUtilities/Source/VPUtilities/Public/Libraries/VPBlueprintLibrary.h
@@ -49,29 +49,32 @@ public:
 	static FTransform GetEditorViewportTransform();
 
 	/** Get the location of the VR HMD */
-	UFUNCTION(BlueprintCallable, Category = "Virtual Production")
+	UE_DEPRECATED(5.7, "VR Editor mode is deprecated; use the XR Creative Framework plugin instead.")
+	UFUNCTION(BlueprintCallable, Category = "Virtual Production", meta = (DeprecatedFunction, DeprecationMessage = "VR Editor mode is deprecated; use the XR Creative Framework plugin instead."))
 	static FTransform GetEditorVRHeadTransform();
 
 	/** Get the VR room transform (the playable area shown as a wireframe cage on Vive and Rift */
-	UFUNCTION(BlueprintCallable, Category = "Virtual Production")
+	UE_DEPRECATED(5.7, "VR Editor mode is deprecated; use the XR Creative Framework plugin instead.")
+	UFUNCTION(BlueprintCallable, Category = "Virtual Production", meta = (DeprecatedFunction, DeprecationMessage = "VR Editor mode is deprecated; use the XR Creative Framework plugin instead."))
 	static FTransform GetEditorVRRoomTransform();
 
 	/** Set the VR grab speed cvar */
-	UFUNCTION(BlueprintCallable, Category = "Virtual Production")
+	UE_DEPRECATED(5.7, "VR Editor mode is deprecated; use the XR Creative Framework plugin instead.")
+	UFUNCTION(BlueprintCallable, Category = "Virtual Production", meta = (DeprecatedFunction, DeprecationMessage = "VR Editor mode is deprecated; use the XR Creative Framework plugin instead."))
 	static void SetGrabSpeed(const float Speed);
 
 	/** Get whether the user is in editor VR mode */
-	UE_DEPRECATED(5.6, "VR Editor mode is deprecated; use the XR Creative Framework plugin instead.")
+	UE_DEPRECATED(5.7, "VR Editor mode is deprecated; use the XR Creative Framework plugin instead.")
 	UFUNCTION(BlueprintCallable, Category = "Virtual Production", meta = (DeprecatedFunction, DeprecationMessage = "VR Editor mode is deprecated; use the XR Creative Framework plugin instead."))
 	static bool IsVREditorModeActive();
 
 	/** Get the hitlocation of the interaction controller's laser pointer, in world space */
-	UE_DEPRECATED(5.6, "VR Editor mode is deprecated; use the XR Creative Framework plugin instead.")
+	UE_DEPRECATED(5.7, "VR Editor mode is deprecated; use the XR Creative Framework plugin instead.")
 	UFUNCTION(BlueprintCallable, Category = "Virtual Production", meta = (DeprecatedFunction, DeprecationMessage = "VR Editor mode is deprecated; use the XR Creative Framework plugin instead."))
 	static FVector GetVREditorLaserHoverLocation();
 
 	/** Get the location of the left interactor, in world space */
-	UE_DEPRECATED(5.6, "VR Editor mode is deprecated; use the XR Creative Framework plugin instead.")
+	UE_DEPRECATED(5.7, "VR Editor mode is deprecated; use the XR Creative Framework plugin instead.")
 	UFUNCTION(BlueprintCallable, Category = "Virtual Production", meta = (DeprecatedFunction, DeprecationMessage = "VR Editor mode is deprecated; use the XR Creative Framework plugin instead."))
 	static FVector GetLeftInteractorLocation();
 
@@ -84,11 +87,13 @@ public:
 	static bool EditorRedo();
 
 	/** Trigger an UnrealEd Duplicate */
-	UFUNCTION(BlueprintCallable, Category = "Virtual Production")
+	UE_DEPRECATED(5.7, "Use UEditorActorSubsystem::DuplicateSelectedActors instead.")
+	UFUNCTION(BlueprintCallable, Category = "Virtual Production", meta = (DeprecatedFunction, DeprecationMessage = "Use UEditorActorSubsystem::DuplicateSelectedActors instead."))
 	static bool EditorDuplicate();
 
 	/** Trigger an UnrealEd Delete */
-	UFUNCTION(BlueprintCallable, Category = "Virtual Production")
+	UE_DEPRECATED(5.7, "Use UEditorActorSubsystem::DeleteSelectedActors instead.")
+	UFUNCTION(BlueprintCallable, Category = "Virtual Production", meta = (DeprecatedFunction, DeprecationMessage = "Use UEditorActorSubsystem::DeleteSelectedActors instead."))
 	static bool EditorDeleteSelectedObjects();
 
 	/**
diff --git a/Engine/Source/Editor/ComponentVisualizers/Private/Manipulator.cpp b/Engine/Source/Editor/ComponentVisualizers/Private/Manipulator.cpp
--- a/Engine/Source/Editor/ComponentVisualizers/Private/Manipulator.cpp
+++ b/Engine/Source/Editor/ComponentVisualizers/Private/Manipulator.cpp
@@ -13,6 +13,9 @@
 
 #include UE_INLINE_GENERATED_CPP_BY_NAME(Manipulator)
 
+PRAGMA_DISABLE_DEPRECATION_WARNINGS
+
+
 class UActorComponent;
 class UViewportInteractor;
 struct FHitResult;
@@ -104,3 +107,6 @@ void AManipulator::SetAssociatedComponent(USceneComponent* SceneComponent)
 		SetActorTransform(AssociatedComponent->GetComponentToWorld());	
 	}
 }
+
+
+PRAGMA_ENABLE_DEPRECATION_WARNINGS
diff --git a/Engine/Source/Editor/ComponentVisualizers/Public/Manipulator.h b/Engine/Source/Editor/ComponentVisualizers/Public/Manipulator.h
--- a/Engine/Source/Editor/ComponentVisualizers/Public/Manipulator.h
+++ b/Engine/Source/Editor/ComponentVisualizers/Public/Manipulator.h
@@ -11,8 +11,12 @@
 
 #include "Manipulator.generated.h"
 
+
 #define UE_API COMPONENTVISUALIZERS_API
 
+PRAGMA_DISABLE_DEPRECATION_WARNINGS
+
+
 class UActorComponent;
 class UObject;
 class USceneComponent;
@@ -20,7 +24,7 @@ class UViewportInteractor;
 struct FHitResult;
 
 UCLASS(MinimalAPI)
-class AManipulator : public AActor, public IViewportInteractableInterface
+class UE_DEPRECATED(5.7, "The ViewportInteraction module is being deprecated alongside VR Editor mode.") AManipulator : public AActor, public IViewportInteractableInterface
 {
 	GENERATED_BODY()
 
@@ -58,4 +62,6 @@ private:
 
 };
 
+PRAGMA_ENABLE_DEPRECATION_WARNINGS
+
 #undef UE_API
diff --git a/Engine/Source/Editor/UnrealEd/Private/LevelEditorViewport.cpp b/Engine/Source/Editor/UnrealEd/Private/LevelEditorViewport.cpp
--- a/Engine/Source/Editor/UnrealEd/Private/LevelEditorViewport.cpp
+++ b/Engine/Source/Editor/UnrealEd/Private/LevelEditorViewport.cpp
@@ -4651,6 +4651,7 @@ EMouseCursor::Type FLevelEditorViewportClient::GetCursor(FViewport* InViewport,i
 
 	if (ViewportWorld != nullptr)
 	{
+PRAGMA_DISABLE_DEPRECATION_WARNINGS
 		// Allow the viewport interaction to override any previously set mouse cursor
 		UViewportWorldInteraction* WorldInteraction = Cast<UViewportWorldInteraction>(GEditor->GetEditorWorldExtensionsManager()->GetEditorWorldExtensions(GetWorld())->FindExtension(UViewportWorldInteraction::StaticClass()));
 		if (WorldInteraction != nullptr)
@@ -4668,6 +4669,7 @@ EMouseCursor::Type FLevelEditorViewportClient::GetCursor(FViewport* InViewport,i
 				UpdateRequiredCursorVisibility();
 			}
 		}
+PRAGMA_ENABLE_DEPRECATION_WARNINGS
 	}
 
 	// Don't select widget axes by mouse over while they're being controlled by a mouse drag.
diff --git a/Engine/Source/Editor/VREditor/Private/UI/VREditorDockableWindow.cpp b/Engine/Source/Editor/VREditor/Private/UI/VREditorDockableWindow.cpp
--- a/Engine/Source/Editor/VREditor/Private/UI/VREditorDockableWindow.cpp
+++ b/Engine/Source/Editor/VREditor/Private/UI/VREditorDockableWindow.cpp
@@ -495,14 +495,15 @@ void AVREditorDockableWindow::SetDockButtonColor(const FLinearColor& LinearColor
 /******************	******************************************************/
 void UDockableWindowDragOperation::ExecuteDrag( UViewportInteractor* Interactor, IViewportInteractableInterface* Interactable  )
 {
-	UVREditorInteractor* VREditorInteractor = Cast<UVREditorInteractor>( Interactor );
-	AVREditorDockableWindow* DockableWindow = Cast<AVREditorDockableWindow>( Interactable );
-
-	if (DockableWindow && !IsValid(DockableWindow))
+	// This is a "manual" interface Cast<> to avoid emitting a spurious deprecation warning.
+	// This code will all be deleted in UE 5.9.
+	if (!Interactable || !IsValid(Interactable->_getUObject()) || !Interactable->_getUObject()->IsA(AVREditorDockableWindow::StaticClass()))
 	{
 		return;
 	}
+	AVREditorDockableWindow* DockableWindow = (AVREditorDockableWindow*)Interactable->_getUObject();
 
+	UVREditorInteractor* VREditorInteractor = Cast<UVREditorInteractor>(Interactor);
 	if ( VREditorInteractor && DockableWindow )
 	{
 		UVREditorUISystem& UISystem = DockableWindow->GetOwner();
diff --git a/Engine/Source/Editor/VREditor/Private/VRModeSettings.cpp b/Engine/Source/Editor/VREditor/Private/VRModeSettings.cpp
--- a/Engine/Source/Editor/VREditor/Private/VRModeSettings.cpp
+++ b/Engine/Source/Editor/VREditor/Private/VRModeSettings.cpp
@@ -10,14 +10,18 @@
 #define LOCTEXT_NAMESPACE "VREditor"
 
 UVRModeSettings::UVRModeSettings()
+PRAGMA_DISABLE_DEPRECATION_WARNINGS
 	: Super()
+PRAGMA_ENABLE_DEPRECATION_WARNINGS
 {
 }
 
 #if WITH_EDITOR
 void UVRModeSettings::PostEditChangeProperty(struct FPropertyChangedEvent& PropertyChangedEvent)
 {
+PRAGMA_DISABLE_DEPRECATION_WARNINGS
 	Super::PostEditChangeProperty(PropertyChangedEvent);
+PRAGMA_ENABLE_DEPRECATION_WARNINGS
 
 	if (PropertyChangedEvent.Property
 		&& PropertyChangedEvent.Property->GetFName() == GET_MEMBER_NAME_CHECKED(UVRModeSettings, bEnableAutoVREditMode)
diff --git a/Engine/Source/Editor/VREditor/Public/VRModeSettings.h b/Engine/Source/Editor/VREditor/Public/VRModeSettings.h
--- a/Engine/Source/Editor/VREditor/Public/VRModeSettings.h
+++ b/Engine/Source/Editor/VREditor/Public/VRModeSettings.h
@@ -29,9 +29,12 @@ enum class UE_DEPRECATED(5.7, "VR Editor mode is deprecated; use the XR Creative
 /**
 * Implements the settings for VR Mode.
 */
+PRAGMA_DISABLE_DEPRECATION_WARNINGS
 UCLASS(MinimalAPI, config = EditorSettings)
 class UVRModeSettings : public UVISettings
 {
+PRAGMA_ENABLE_DEPRECATION_WARNINGS
+
 	GENERATED_BODY()
 
 public:
diff --git a/Engine/Source/Editor/ViewportInteraction/Private/ActorTransformer.cpp b/Engine/Source/Editor/ViewportInteraction/Private/ActorTransformer.cpp
--- a/Engine/Source/Editor/ViewportInteraction/Private/ActorTransformer.cpp
+++ b/Engine/Source/Editor/ViewportInteraction/Private/ActorTransformer.cpp
@@ -13,6 +13,8 @@
 
 #include UE_INLINE_GENERATED_CPP_BY_NAME(ActorTransformer)
 
+PRAGMA_DISABLE_DEPRECATION_WARNINGS
+
 
 void UActorTransformer::Init( UViewportWorldInteraction* InitViewportWorldInteraction )
 {
@@ -169,3 +171,5 @@ void UActorTransformer::OnActorSelectionChanged( UObject* ChangedObject )
 	ViewportWorldInteraction->SetTransformables( MoveTemp( NewTransformables ), bNewObjectsSelected );
 }
 
+
+PRAGMA_ENABLE_DEPRECATION_WARNINGS
diff --git a/Engine/Source/Editor/ViewportInteraction/Private/Gizmo/VIBaseTransformGizmo.cpp b/Engine/Source/Editor/ViewportInteraction/Private/Gizmo/VIBaseTransformGizmo.cpp
--- a/Engine/Source/Editor/ViewportInteraction/Private/Gizmo/VIBaseTransformGizmo.cpp
+++ b/Engine/Source/Editor/ViewportInteraction/Private/Gizmo/VIBaseTransformGizmo.cpp
@@ -11,6 +11,9 @@
 
 #include UE_INLINE_GENERATED_CPP_BY_NAME(VIBaseTransformGizmo)
 
+PRAGMA_DISABLE_DEPRECATION_WARNINGS
+
+
 namespace VREd
 {
 	static FAutoConsoleVariable GizmoSelectionAnimationDuration( TEXT( "VI.GizmoSelectionAnimationDuration" ), 0.15f, TEXT( "How long to animate the gizmo after objects are selected" ) );
@@ -262,3 +265,6 @@ void FTransformGizmoHandlePlacement::GetCenterHandleCountAndFacingAxisIndex( int
 		OutCenterAxisIndex = INDEX_NONE;
 	}
 }
+
+
+PRAGMA_ENABLE_DEPRECATION_WARNINGS
diff --git a/Engine/Source/Editor/ViewportInteraction/Private/Gizmo/VIGizmoHandle.cpp b/Engine/Source/Editor/ViewportInteraction/Private/Gizmo/VIGizmoHandle.cpp
--- a/Engine/Source/Editor/ViewportInteraction/Private/Gizmo/VIGizmoHandle.cpp
+++ b/Engine/Source/Editor/ViewportInteraction/Private/Gizmo/VIGizmoHandle.cpp
@@ -12,6 +12,9 @@
 
 #include UE_INLINE_GENERATED_CPP_BY_NAME(VIGizmoHandle)
 
+PRAGMA_DISABLE_DEPRECATION_WARNINGS
+
+
 UGizmoHandleGroup::UGizmoHandleGroup()
 	: Super(),
 	GizmoMaterial(nullptr),
@@ -431,3 +434,6 @@ void UAxisGizmoHandleGroup::UpdateHandlesRelativeTransformOnAxis( const FTransfo
 		}
 	}
 }
+
+
+PRAGMA_ENABLE_DEPRECATION_WARNINGS
diff --git a/Engine/Source/Editor/ViewportInteraction/Private/Gizmo/VIGizmoHandleMeshComponent.cpp b/Engine/Source/Editor/ViewportInteraction/Private/Gizmo/VIGizmoHandleMeshComponent.cpp
--- a/Engine/Source/Editor/ViewportInteraction/Private/Gizmo/VIGizmoHandleMeshComponent.cpp
+++ b/Engine/Source/Editor/ViewportInteraction/Private/Gizmo/VIGizmoHandleMeshComponent.cpp
@@ -8,6 +8,9 @@
 
 #include UE_INLINE_GENERATED_CPP_BY_NAME(VIGizmoHandleMeshComponent)
 
+PRAGMA_DISABLE_DEPRECATION_WARNINGS
+
+
 class FGimzoHandleSceneProxy final : public FStaticMeshSceneProxy
 {
 public:
@@ -89,3 +92,6 @@ FPrimitiveSceneProxy* UGizmoHandleMeshComponent::CreateSceneProxy()
 
 	return new FGimzoHandleSceneProxy(this);
 }
+
+
+PRAGMA_ENABLE_DEPRECATION_WARNINGS
diff --git a/Engine/Source/Editor/ViewportInteraction/Private/Gizmo/VIPivotTransformGizmo.cpp b/Engine/Source/Editor/ViewportInteraction/Private/Gizmo/VIPivotTransformGizmo.cpp
--- a/Engine/Source/Editor/ViewportInteraction/Private/Gizmo/VIPivotTransformGizmo.cpp
+++ b/Engine/Source/Editor/ViewportInteraction/Private/Gizmo/VIPivotTransformGizmo.cpp
@@ -23,6 +23,9 @@
 
 #include UE_INLINE_GENERATED_CPP_BY_NAME(VIPivotTransformGizmo)
 
+PRAGMA_DISABLE_DEPRECATION_WARNINGS
+
+
 namespace VREd //@todo VREditor: Duplicates of TransformGizmo
 {
 	// @todo vreditor tweak: Tweak out console variables
@@ -533,3 +536,6 @@ void UPivotRotationGizmoHandleGroup::SetIndicatorColor(UStaticMeshComponent* InM
 	MID0->SetVectorParameterValue(StaticHandleColorParameter, InHandleColor);
 	MID1->SetVectorParameterValue(StaticHandleColorParameter, InHandleColor);
 }
+
+
+PRAGMA_ENABLE_DEPRECATION_WARNINGS
diff --git a/Engine/Source/Editor/ViewportInteraction/Private/Gizmo/VIPivotTransformGizmo.h b/Engine/Source/Editor/ViewportInteraction/Private/Gizmo/VIPivotTransformGizmo.h
--- a/Engine/Source/Editor/ViewportInteraction/Private/Gizmo/VIPivotTransformGizmo.h
+++ b/Engine/Source/Editor/ViewportInteraction/Private/Gizmo/VIPivotTransformGizmo.h
@@ -8,11 +8,15 @@
 #include "VIGizmoHandle.h"
 #include "VIPivotTransformGizmo.generated.h"
 
+
+PRAGMA_DISABLE_DEPRECATION_WARNINGS
+
+
 /**
  * A transform gizmo on the pivot that allows you to interact with selected objects by moving, scaling and rotating.
  */
 UCLASS()
-class APivotTransformGizmo : public ABaseTransformGizmo
+class UE_DEPRECATED(5.7, "The ViewportInteraction module is being deprecated alongside VR Editor mode.") APivotTransformGizmo : public ABaseTransformGizmo
 {
 	GENERATED_BODY()
 
@@ -65,7 +69,7 @@ private:
  * Axis Gizmo handle for translating
  */
 UCLASS()
-class UPivotTranslationGizmoHandleGroup : public UAxisGizmoHandleGroup
+class UE_DEPRECATED(5.7, "The ViewportInteraction module is being deprecated alongside VR Editor mode.") UPivotTranslationGizmoHandleGroup : public UAxisGizmoHandleGroup
 {
 	GENERATED_BODY()
 
@@ -86,7 +90,7 @@ public:
  * Axis Gizmo handle for scaling
  */
 UCLASS()
-class UPivotScaleGizmoHandleGroup : public UAxisGizmoHandleGroup
+class UE_DEPRECATED(5.7, "The ViewportInteraction module is being deprecated alongside VR Editor mode.") UPivotScaleGizmoHandleGroup : public UAxisGizmoHandleGroup
 {
 	GENERATED_BODY()
 
@@ -110,7 +114,7 @@ public:
  * Axis Gizmo handle for plane translation
  */
 UCLASS()
-class UPivotPlaneTranslationGizmoHandleGroup : public UAxisGizmoHandleGroup
+class UE_DEPRECATED(5.7, "The ViewportInteraction module is being deprecated alongside VR Editor mode.") UPivotPlaneTranslationGizmoHandleGroup : public UAxisGizmoHandleGroup
 {
 	GENERATED_BODY()
 
@@ -132,7 +136,7 @@ public:
  * Axis Gizmo handle for rotation
  */
 UCLASS()
-class UPivotRotationGizmoHandleGroup : public UAxisGizmoHandleGroup
+class UE_DEPRECATED(5.7, "The ViewportInteraction module is being deprecated alongside VR Editor mode.") UPivotRotationGizmoHandleGroup : public UAxisGizmoHandleGroup
 {
 	GENERATED_BODY()
 
@@ -187,3 +191,6 @@ private:
 	/** The rotation when starting to drag the gizmo */
 	TOptional<FQuat> StartDragRotation;
 };
+
+
+PRAGMA_ENABLE_DEPRECATION_WARNINGS
diff --git a/Engine/Source/Editor/ViewportInteraction/Private/Gizmo/VIStretchGizmoHandle.cpp b/Engine/Source/Editor/ViewportInteraction/Private/Gizmo/VIStretchGizmoHandle.cpp
--- a/Engine/Source/Editor/ViewportInteraction/Private/Gizmo/VIStretchGizmoHandle.cpp
+++ b/Engine/Source/Editor/ViewportInteraction/Private/Gizmo/VIStretchGizmoHandle.cpp
@@ -9,6 +9,9 @@
 
 #include UE_INLINE_GENERATED_CPP_BY_NAME(VIStretchGizmoHandle)
 
+PRAGMA_DISABLE_DEPRECATION_WARNINGS
+
+
 UStretchGizmoHandleGroup::UStretchGizmoHandleGroup()
 	: Super()
 {
@@ -317,3 +320,6 @@ void UStretchGizmoHandleDragOperation::ExecuteDrag(FDraggingTransformableData& D
 		DraggingData.bAllowSnap = false;
 	}
 }
+
+
+PRAGMA_ENABLE_DEPRECATION_WARNINGS
diff --git a/Engine/Source/Editor/ViewportInteraction/Private/Gizmo/VIStretchGizmoHandle.h b/Engine/Source/Editor/ViewportInteraction/Private/Gizmo/VIStretchGizmoHandle.h
--- a/Engine/Source/Editor/ViewportInteraction/Private/Gizmo/VIStretchGizmoHandle.h
+++ b/Engine/Source/Editor/ViewportInteraction/Private/Gizmo/VIStretchGizmoHandle.h
@@ -10,13 +10,16 @@
 
 #define UE_API VIEWPORTINTERACTION_API
 
+PRAGMA_DISABLE_DEPRECATION_WARNINGS
+
+
 enum class EGizmoHandleTypes : uint8;
 
 /**
  * Gizmo handle for stretching/scaling
  */
 UCLASS(MinimalAPI)
-class UStretchGizmoHandleGroup : public UGizmoHandleGroup
+class UE_DEPRECATED(5.7, "The ViewportInteraction module is being deprecated alongside VR Editor mode.") UStretchGizmoHandleGroup : public UGizmoHandleGroup
 {
 	GENERATED_BODY()
 
@@ -47,4 +50,7 @@ public:
 	virtual void ExecuteDrag(struct FDraggingTransformableData& DraggingData) override;
 };
 
+
+PRAGMA_ENABLE_DEPRECATION_WARNINGS
+
 #undef UE_API
diff --git a/Engine/Source/Editor/ViewportInteraction/Private/Gizmo/VIUniformScaleGizmoHandle.cpp b/Engine/Source/Editor/ViewportInteraction/Private/Gizmo/VIUniformScaleGizmoHandle.cpp
--- a/Engine/Source/Editor/ViewportInteraction/Private/Gizmo/VIUniformScaleGizmoHandle.cpp
+++ b/Engine/Source/Editor/ViewportInteraction/Private/Gizmo/VIUniformScaleGizmoHandle.cpp
@@ -13,6 +13,9 @@
 
 #include UE_INLINE_GENERATED_CPP_BY_NAME(VIUniformScaleGizmoHandle)
 
+PRAGMA_DISABLE_DEPRECATION_WARNINGS
+
+
 UUniformScaleGizmoHandleGroup::UUniformScaleGizmoHandleGroup()
 	: Super(),
 	bUsePivotAsLocation( true )
@@ -108,3 +111,6 @@ void UUniformScaleGizmoHandleGroup::SetUsePivotPointAsLocation( const bool bInUs
 {
 	bUsePivotAsLocation = bInUsePivotAsLocation;
 }
+
+
+PRAGMA_ENABLE_DEPRECATION_WARNINGS
diff --git a/Engine/Source/Editor/ViewportInteraction/Private/Gizmo/VIUniformScaleGizmoHandle.h b/Engine/Source/Editor/ViewportInteraction/Private/Gizmo/VIUniformScaleGizmoHandle.h
--- a/Engine/Source/Editor/ViewportInteraction/Private/Gizmo/VIUniformScaleGizmoHandle.h
+++ b/Engine/Source/Editor/ViewportInteraction/Private/Gizmo/VIUniformScaleGizmoHandle.h
@@ -9,13 +9,16 @@
 
 #define UE_API VIEWPORTINTERACTION_API
 
+PRAGMA_DISABLE_DEPRECATION_WARNINGS
+
+
 enum class EGizmoHandleTypes : uint8;
 
 /**
  * Gizmo handle for uniform scaling
  */
 UCLASS(MinimalAPI)
-class UUniformScaleGizmoHandleGroup : public UGizmoHandleGroup
+class UE_DEPRECATED(5.7, "The ViewportInteraction module is being deprecated alongside VR Editor mode.") UUniformScaleGizmoHandleGroup : public UGizmoHandleGroup
 {
 	GENERATED_BODY()
 
@@ -40,4 +43,7 @@ private:
 	bool bUsePivotAsLocation;
 };
 
+
+PRAGMA_ENABLE_DEPRECATION_WARNINGS
+
 #undef UE_API
diff --git a/Engine/Source/Editor/ViewportInteraction/Private/MouseCursorInteractor.cpp b/Engine/Source/Editor/ViewportInteraction/Private/MouseCursorInteractor.cpp
--- a/Engine/Source/Editor/ViewportInteraction/Private/MouseCursorInteractor.cpp
+++ b/Engine/Source/Editor/ViewportInteraction/Private/MouseCursorInteractor.cpp
@@ -8,6 +8,8 @@
 
 #include UE_INLINE_GENERATED_CPP_BY_NAME(MouseCursorInteractor)
 
+PRAGMA_DISABLE_DEPRECATION_WARNINGS
+
 
 UMouseCursorInteractor::UMouseCursorInteractor()
 	: UViewportInteractor()
@@ -90,3 +92,5 @@ bool UMouseCursorInteractor::AllowLaserSmoothing() const
 	return false;
 }
 
+
+PRAGMA_ENABLE_DEPRECATION_WARNINGS
diff --git a/Engine/Source/Editor/ViewportInteraction/Private/ViewportDragOperation.cpp b/Engine/Source/Editor/ViewportInteraction/Private/ViewportDragOperation.cpp
--- a/Engine/Source/Editor/ViewportInteraction/Private/ViewportDragOperation.cpp
+++ b/Engine/Source/Editor/ViewportInteraction/Private/ViewportDragOperation.cpp
@@ -6,6 +6,9 @@
 
 #include UE_INLINE_GENERATED_CPP_BY_NAME(ViewportDragOperation)
 
+PRAGMA_DISABLE_DEPRECATION_WARNINGS
+
+
 UViewportDragOperationComponent::UViewportDragOperationComponent() : 
 	Super(),
 	DragOperation(nullptr)
@@ -48,3 +51,6 @@ bool UViewportDragOperationComponent::IsDragging() const
 {
 	return DragOperation != nullptr;
 }
+
+
+PRAGMA_ENABLE_DEPRECATION_WARNINGS
diff --git a/Engine/Source/Editor/ViewportInteraction/Private/ViewportInteractionDragOperations.cpp b/Engine/Source/Editor/ViewportInteraction/Private/ViewportInteractionDragOperations.cpp
--- a/Engine/Source/Editor/ViewportInteraction/Private/ViewportInteractionDragOperations.cpp
+++ b/Engine/Source/Editor/ViewportInteraction/Private/ViewportInteractionDragOperations.cpp
@@ -11,6 +11,9 @@
 
 #include UE_INLINE_GENERATED_CPP_BY_NAME(ViewportInteractionDragOperations)
 
+PRAGMA_DISABLE_DEPRECATION_WARNINGS
+
+
 namespace VI
 {
 	static FAutoConsoleVariable ScaleSensitivity(TEXT("VI.ScaleSensitivity"), 0.005f, TEXT("Sensitivity for scaling"));
@@ -166,3 +169,6 @@ void UUniformScaleDragOperation::ExecuteDrag(struct FDraggingTransformableData&
 	DraggingData.bOutShouldApplyVelocitiesFromDrag = true;
 	DraggingData.bOutScaled = true;
 }
+
+
+PRAGMA_ENABLE_DEPRECATION_WARNINGS
diff --git a/Engine/Source/Editor/ViewportInteraction/Private/ViewportInteractionInputProcessor.cpp b/Engine/Source/Editor/ViewportInteraction/Private/ViewportInteractionInputProcessor.cpp
--- a/Engine/Source/Editor/ViewportInteraction/Private/ViewportInteractionInputProcessor.cpp
+++ b/Engine/Source/Editor/ViewportInteraction/Private/ViewportInteractionInputProcessor.cpp
@@ -5,6 +5,10 @@
 #include "Misc/App.h"
 #include "ViewportWorldInteraction.h"
 
+
+PRAGMA_DISABLE_DEPRECATION_WARNINGS
+
+
 FViewportInteractionInputProcessor::FViewportInteractionInputProcessor(UViewportWorldInteraction* InWorldInteraction)
 	: WorldInteraction(InWorldInteraction)
 {
@@ -39,3 +43,5 @@ bool FViewportInteractionInputProcessor::HandleMouseMoveEvent( FSlateApplication
 	return false;
 }
 
+
+PRAGMA_ENABLE_DEPRECATION_WARNINGS
diff --git a/Engine/Source/Editor/ViewportInteraction/Private/ViewportInteractionInputProcessor.h b/Engine/Source/Editor/ViewportInteraction/Private/ViewportInteractionInputProcessor.h
--- a/Engine/Source/Editor/ViewportInteraction/Private/ViewportInteractionInputProcessor.h
+++ b/Engine/Source/Editor/ViewportInteraction/Private/ViewportInteractionInputProcessor.h
@@ -6,12 +6,17 @@
 #include "GenericPlatform/ICursor.h"
 #include "Framework/Application/IInputProcessor.h"
 
+
 class FSlateApplication;
 class FViewportWorldInteractionManager;
 struct FAnalogInputEvent;
 struct FKeyEvent;
 struct FPointerEvent;
 
+
+PRAGMA_DISABLE_DEPRECATION_WARNINGS
+
+
 class FViewportInteractionInputProcessor : public IInputProcessor
 {
 public:
@@ -31,3 +36,6 @@ private:
 	/** The WorldInteraction that will receive the input */
 	class UViewportWorldInteraction* WorldInteraction;
 };
+
+
+PRAGMA_ENABLE_DEPRECATION_WARNINGS
diff --git a/Engine/Source/Editor/ViewportInteraction/Private/ViewportInteractionModule.cpp b/Engine/Source/Editor/ViewportInteraction/Private/ViewportInteractionModule.cpp
--- a/Engine/Source/Editor/ViewportInteraction/Private/ViewportInteractionModule.cpp
+++ b/Engine/Source/Editor/ViewportInteraction/Private/ViewportInteractionModule.cpp
@@ -9,6 +9,10 @@
 #include "LevelEditor.h"
 #include "SLevelViewport.h"
 
+
+PRAGMA_DISABLE_DEPRECATION_WARNINGS
+
+
 namespace VI
 {
 	static FAutoConsoleCommand ForceMode(TEXT("VI.ForceMode"), TEXT("Toggles viewport interaction on desktop"), FConsoleCommandDelegate::CreateStatic(&FViewportInteractionModule::ToggleMode));
@@ -97,3 +101,6 @@ void FViewportInteractionModule::ToggleMode()
 		}
 	}
 }
+
+
+PRAGMA_ENABLE_DEPRECATION_WARNINGS
diff --git a/Engine/Source/Editor/ViewportInteraction/Private/ViewportInteractionModule.h b/Engine/Source/Editor/ViewportInteraction/Private/ViewportInteractionModule.h
--- a/Engine/Source/Editor/ViewportInteraction/Private/ViewportInteractionModule.h
+++ b/Engine/Source/Editor/ViewportInteraction/Private/ViewportInteractionModule.h
@@ -5,7 +5,10 @@
 #include "CoreMinimal.h"
 #include "IViewportInteractionModule.h"
 
-class FViewportInteractionModule : public IViewportInteractionModule
+class FViewportInteractionModule
+PRAGMA_DISABLE_DEPRECATION_WARNINGS
+	: public IViewportInteractionModule
+PRAGMA_ENABLE_DEPRECATION_WARNINGS
 {
 public:
 	
diff --git a/Engine/Source/Editor/ViewportInteraction/Private/ViewportInteractionStyle.cpp b/Engine/Source/Editor/ViewportInteraction/Private/ViewportInteractionStyle.cpp
--- a/Engine/Source/Editor/ViewportInteraction/Private/ViewportInteractionStyle.cpp
+++ b/Engine/Source/Editor/ViewportInteraction/Private/ViewportInteractionStyle.cpp
@@ -10,6 +10,10 @@
 #include "Rendering/SlateRenderer.h"
 #include "Styling/SlateStyleRegistry.h"
 
+
+PRAGMA_DISABLE_DEPRECATION_WARNINGS
+
+
 class ISlateStyle;
 
 TSharedPtr< FSlateStyleSet > FViewportInteractionStyle::ViewportInteractionStyleInstance = NULL;
@@ -71,3 +75,6 @@ const ISlateStyle& FViewportInteractionStyle::Get()
 #undef IMAGE_BRUSH
 #undef BOX_BRUSH
 #undef BORDER_BRUSH
+
+
+PRAGMA_ENABLE_DEPRECATION_WARNINGS
diff --git a/Engine/Source/Editor/ViewportInteraction/Private/ViewportInteractor.cpp b/Engine/Source/Editor/ViewportInteraction/Private/ViewportInteractor.cpp
--- a/Engine/Source/Editor/ViewportInteraction/Private/ViewportInteractor.cpp
+++ b/Engine/Source/Editor/ViewportInteraction/Private/ViewportInteractor.cpp
@@ -15,6 +15,9 @@
 
 #include UE_INLINE_GENERATED_CPP_BY_NAME(ViewportInteractor)
 
+PRAGMA_DISABLE_DEPRECATION_WARNINGS
+
+
 #define LOCTEXT_NAMESPACE "ViewportInteractor"
 
 namespace VI
@@ -201,8 +204,13 @@ bool UViewportInteractor::HandleInputKey( FEditorViewportClient& ViewportClient,
 							}
 
 							bool bCanBeSelected = true;
-							if ( IViewportInteractableInterface* ActorInteractable = Cast<IViewportInteractableInterface>( Actor ) )
+							// This is a "manual" interface Cast<> to avoid emitting a spurious deprecation warning.
+							// This code will all be deleted in UE 5.9.
+							if (Actor->GetClass()->ImplementsInterface(UViewportInteractableInterface::StaticClass()))
 							{
+								IViewportInteractableInterface* ActorInteractable =
+									(IViewportInteractableInterface*)Actor->GetInterfaceAddress(UViewportInteractableInterface::StaticClass());
+
 								bCanBeSelected = ActorInteractable->CanBeSelected();
 								if (!bCanBeSelected)
 								{
@@ -917,3 +925,6 @@ bool UViewportInteractor::CanCarry() const
 }
 
 #undef LOCTEXT_NAMESPACE
+
+
+PRAGMA_ENABLE_DEPRECATION_WARNINGS
diff --git a/Engine/Source/Editor/ViewportInteraction/Private/ViewportTransformer.cpp b/Engine/Source/Editor/ViewportInteraction/Private/ViewportTransformer.cpp
--- a/Engine/Source/Editor/ViewportInteraction/Private/ViewportTransformer.cpp
+++ b/Engine/Source/Editor/ViewportInteraction/Private/ViewportTransformer.cpp
@@ -5,6 +5,8 @@
 
 #include UE_INLINE_GENERATED_CPP_BY_NAME(ViewportTransformer)
 
+PRAGMA_DISABLE_DEPRECATION_WARNINGS
+
 
 void UViewportTransformer::Init( UViewportWorldInteraction* InitViewportWorldInteraction )
 {
@@ -17,3 +19,5 @@ void UViewportTransformer::Shutdown()
 	this->ViewportWorldInteraction = nullptr;
 }
 
+
+PRAGMA_ENABLE_DEPRECATION_WARNINGS
diff --git a/Engine/Source/Editor/ViewportInteraction/Private/ViewportWorldInteraction.cpp b/Engine/Source/Editor/ViewportInteraction/Private/ViewportWorldInteraction.cpp
--- a/Engine/Source/Editor/ViewportInteraction/Private/ViewportWorldInteraction.cpp
+++ b/Engine/Source/Editor/ViewportInteraction/Private/ViewportWorldInteraction.cpp
@@ -49,6 +49,8 @@
 
 #include UE_INLINE_GENERATED_CPP_BY_NAME(ViewportWorldInteraction)
 
+PRAGMA_DISABLE_DEPRECATION_WARNINGS
+
 
 #define LOCTEXT_NAMESPACE "ViewportWorldInteraction"
 
@@ -876,9 +878,13 @@ void UViewportWorldInteraction::HoverTick( const float DeltaTime )
 					}
 					else
 					{
-						IViewportInteractableInterface* Interactable = Cast<IViewportInteractableInterface>( Actor );
-						if ( Interactable )
+						// This is a "manual" interface Cast<> to avoid emitting a spurious deprecation warning.
+						// This code will all be deleted in UE 5.9.
+						if (Actor->GetClass()->ImplementsInterface(UViewportInteractableInterface::StaticClass()))
 						{
+							IViewportInteractableInterface* Interactable =
+								(IViewportInteractableInterface*)Actor->GetInterfaceAddress(UViewportInteractableInterface::StaticClass());
+
 							NewHoveredActorComponent = HoveredActorComponent;
 
 							// Check if the current hovered component of the interactor is different from the hitresult component
@@ -901,9 +907,13 @@ void UViewportWorldInteraction::HoverTick( const float DeltaTime )
 				AActor* HoveredActor = InteractorData.LastHoveredActorComponent->GetOwner();
 				if ( HoveredActor )
 				{
-					IViewportInteractableInterface* Interactable = Cast<IViewportInteractableInterface>( HoveredActor );
-					if ( Interactable )
+					// This is a "manual" interface Cast<> to avoid emitting a spurious deprecation warning.
+					// This code will all be deleted in UE 5.9.
+					if (HoveredActor->GetClass()->ImplementsInterface(UViewportInteractableInterface::StaticClass()))
 					{
+						IViewportInteractableInterface* Interactable =
+							(IViewportInteractableInterface*)HoveredActor->GetInterfaceAddress(UViewportInteractableInterface::StaticClass());
+
 						Interactable->OnHoverLeave( Interactor, NewHoveredActorComponent );
 					}
 				}
@@ -3830,3 +3840,5 @@ void UViewportWorldInteraction::UseVWInteractions()
 
 #undef LOCTEXT_NAMESPACE
 
+
+PRAGMA_ENABLE_DEPRECATION_WARNINGS
diff --git a/Engine/Source/Editor/ViewportInteraction/Public/ActorTransformer.h b/Engine/Source/Editor/ViewportInteraction/Public/ActorTransformer.h
--- a/Engine/Source/Editor/ViewportInteraction/Public/ActorTransformer.h
+++ b/Engine/Source/Editor/ViewportInteraction/Public/ActorTransformer.h
@@ -7,9 +7,11 @@
 
 #define UE_API VIEWPORTINTERACTION_API
 
+PRAGMA_DISABLE_DEPRECATION_WARNINGS
+
 
 UCLASS(MinimalAPI)
-class UActorTransformer : public UViewportTransformer
+class UE_DEPRECATED(5.7, "The ViewportInteraction module is being deprecated alongside VR Editor mode.") UActorTransformer : public UViewportTransformer
 {
 	GENERATED_BODY()
 
@@ -32,4 +34,7 @@ protected:
 
 };
 
+
+PRAGMA_ENABLE_DEPRECATION_WARNINGS
+
 #undef UE_API
diff --git a/Engine/Source/Editor/ViewportInteraction/Public/ActorViewportTransformable.h b/Engine/Source/Editor/ViewportInteraction/Public/ActorViewportTransformable.h
--- a/Engine/Source/Editor/ViewportInteraction/Public/ActorViewportTransformable.h
+++ b/Engine/Source/Editor/ViewportInteraction/Public/ActorViewportTransformable.h
@@ -8,10 +8,13 @@
 
 #define UE_API VIEWPORTINTERACTION_API
 
+PRAGMA_DISABLE_DEPRECATION_WARNINGS
+
+
 /**
  * A transformable actor
  */
-class FActorViewportTransformable : public FViewportTransformable
+class UE_DEPRECATED(5.7, "The ViewportInteraction module is being deprecated alongside VR Editor mode.") FActorViewportTransformable : public FViewportTransformable
 {
 
 public:
@@ -40,4 +43,6 @@ public:
 };
 
 
+PRAGMA_ENABLE_DEPRECATION_WARNINGS
+
 #undef UE_API
diff --git a/Engine/Source/Editor/ViewportInteraction/Public/IViewportInteractionModule.h b/Engine/Source/Editor/ViewportInteraction/Public/IViewportInteractionModule.h
--- a/Engine/Source/Editor/ViewportInteraction/Public/IViewportInteractionModule.h
+++ b/Engine/Source/Editor/ViewportInteraction/Public/IViewportInteractionModule.h
@@ -9,7 +9,7 @@
 /**
  * The public interface to this module
  */
-class IViewportInteractionModule : public IModuleInterface
+class UE_DEPRECATED(5.7, "The ViewportInteraction module is being deprecated alongside VR Editor mode.") IViewportInteractionModule : public IModuleInterface
 {
 
 public:
diff --git a/Engine/Source/Editor/ViewportInteraction/Public/MouseCursorInteractor.h b/Engine/Source/Editor/ViewportInteraction/Public/MouseCursorInteractor.h
--- a/Engine/Source/Editor/ViewportInteraction/Public/MouseCursorInteractor.h
+++ b/Engine/Source/Editor/ViewportInteraction/Public/MouseCursorInteractor.h
@@ -12,9 +12,12 @@
 /**
  * Viewport interactor for a mouse cursor
  */
+PRAGMA_DISABLE_DEPRECATION_WARNINGS
 UCLASS(MinimalAPI)
-class UMouseCursorInteractor : public UViewportInteractor
+class UE_DEPRECATED(5.7, "The ViewportInteraction module is being deprecated alongside VR Editor mode.") UMouseCursorInteractor : public UViewportInteractor
 {
+PRAGMA_ENABLE_DEPRECATION_WARNINGS
+
 	GENERATED_BODY()
 
 public:
diff --git a/Engine/Source/Editor/ViewportInteraction/Public/VIBaseTransformGizmo.h b/Engine/Source/Editor/ViewportInteraction/Public/VIBaseTransformGizmo.h
--- a/Engine/Source/Editor/ViewportInteraction/Public/VIBaseTransformGizmo.h
+++ b/Engine/Source/Editor/ViewportInteraction/Public/VIBaseTransformGizmo.h
@@ -9,12 +9,16 @@
 #include "ViewportInteractionTypes.h"
 #include "VIBaseTransformGizmo.generated.h"
 
+
 #define UE_API VIEWPORTINTERACTION_API
 
+PRAGMA_DISABLE_DEPRECATION_WARNINGS
+
+
 class UMaterialInterface;
 
 UENUM()
-enum class EGizmoHandleTypes : uint8
+enum class UE_DEPRECATED(5.7, "The ViewportInteraction module is being deprecated alongside VR Editor mode.") EGizmoHandleTypes : uint8
 {
 	All = 0,
 	Translate = 1,
@@ -26,7 +30,7 @@ enum class EGizmoHandleTypes : uint8
 * Displays measurements along the bounds of selected objects
 */
 USTRUCT()
-struct FTransformGizmoMeasurement
+struct UE_DEPRECATED(5.7, "The ViewportInteraction module is being deprecated alongside VR Editor mode.") FTransformGizmoMeasurement
 {
 	GENERATED_BODY()
 
@@ -39,8 +43,8 @@ struct FTransformGizmoMeasurement
 /**
  * Base class for transform gizmo
  */
-UCLASS(MinimalAPI,  Abstract, NotPlaceable )
-class ABaseTransformGizmo : public AActor
+UCLASS(MinimalAPI, Abstract, NotPlaceable)
+class UE_DEPRECATED(5.7, "The ViewportInteraction module is being deprecated alongside VR Editor mode.") ABaseTransformGizmo : public AActor
 {
 	GENERATED_BODY()
 
@@ -113,4 +117,7 @@ protected:
 	EGizmoHandleTypes GizmoType;
 };	
 
+
+PRAGMA_ENABLE_DEPRECATION_WARNINGS
+
 #undef UE_API
diff --git a/Engine/Source/Editor/ViewportInteraction/Public/VIGizmoHandle.h b/Engine/Source/Editor/ViewportInteraction/Public/VIGizmoHandle.h
--- a/Engine/Source/Editor/ViewportInteraction/Public/VIGizmoHandle.h
+++ b/Engine/Source/Editor/ViewportInteraction/Public/VIGizmoHandle.h
@@ -11,13 +11,16 @@
 
 #define UE_API VIEWPORTINTERACTION_API
 
+PRAGMA_DISABLE_DEPRECATION_WARNINGS
+
+
 class UMaterialInterface;
 class UStaticMesh;
 enum class EGizmoHandleTypes : uint8;
 class UActorComponent;
 
 USTRUCT()
-struct FGizmoHandle
+struct UE_DEPRECATED(5.7, "The ViewportInteraction module is being deprecated alongside VR Editor mode.") FGizmoHandle
 {
 	GENERATED_BODY()
 
@@ -33,7 +36,7 @@ struct FGizmoHandle
  * Base class for gizmo handles
  */
 UCLASS(MinimalAPI,  ABSTRACT )
-class UGizmoHandleGroup : public USceneComponent
+class UE_DEPRECATED(5.7, "The ViewportInteraction module is being deprecated alongside VR Editor mode.") UGizmoHandleGroup : public USceneComponent
 {
 	GENERATED_BODY()
 
@@ -137,11 +140,12 @@ private:
 	bool bShowOnUniversalGizmo;
 };
 
+
 /**
  * Base class for gizmo handles on axis
  */
 UCLASS(MinimalAPI, ABSTRACT)
-class UAxisGizmoHandleGroup : public UGizmoHandleGroup
+class UE_DEPRECATED(5.7, "The ViewportInteraction module is being deprecated alongside VR Editor mode.") UAxisGizmoHandleGroup : public UGizmoHandleGroup
 {
 	GENERATED_BODY()
 
@@ -158,4 +162,7 @@ protected:
 		const FVector& ViewLocation, class UActorComponent* DraggingHandle, const TArray< UActorComponent* >& HoveringOverHandles );
 };
 
+
+PRAGMA_ENABLE_DEPRECATION_WARNINGS
+
 #undef UE_API
diff --git a/Engine/Source/Editor/ViewportInteraction/Public/VIGizmoHandleMeshComponent.h b/Engine/Source/Editor/ViewportInteraction/Public/VIGizmoHandleMeshComponent.h
--- a/Engine/Source/Editor/ViewportInteraction/Public/VIGizmoHandleMeshComponent.h
+++ b/Engine/Source/Editor/ViewportInteraction/Public/VIGizmoHandleMeshComponent.h
@@ -7,7 +7,7 @@
 #include "VIGizmoHandleMeshComponent.generated.h"
 
 UCLASS(hidecategories = Object)
-class UGizmoHandleMeshComponent : public UStaticMeshComponent
+class UE_DEPRECATED(5.7, "The ViewportInteraction module is being deprecated alongside VR Editor mode.") UGizmoHandleMeshComponent : public UStaticMeshComponent
 {
 	GENERATED_BODY()
 
diff --git a/Engine/Source/Editor/ViewportInteraction/Public/VISettings.h b/Engine/Source/Editor/ViewportInteraction/Public/VISettings.h
--- a/Engine/Source/Editor/ViewportInteraction/Public/VISettings.h
+++ b/Engine/Source/Editor/ViewportInteraction/Public/VISettings.h
@@ -11,7 +11,7 @@
 * Implements the settings for Viewport Interaction and VR Mode. Note that this is primarily a base class for anything that needs to be accessed in the VI module.
 */
 UCLASS(MinimalAPI, config = EditorSettings)
-class UVISettings : public UObject
+class UE_DEPRECATED(5.7, "The ViewportInteraction module is being deprecated alongside VR Editor mode.") UVISettings : public UObject
 {
 	GENERATED_BODY()
 
@@ -24,15 +24,18 @@ public:
 	};
 
 	/** Whether the world should scale relative to your tracking space floor instead of the center of your hand locations */
-	UPROPERTY(EditAnywhere, config, Category = "World Movement")
+	UE_DEPRECATED(5.7, "The ViewportInteraction module is being deprecated alongside VR Editor mode.")
+	UPROPERTY(config)
 	uint32 bScaleWorldFromFloor : 1;
 
 	/** Whether to compute a new center point for scaling relative from by looking at how far either controller moved relative to the last frame */
-	UPROPERTY(EditAnywhere, config, Category = "World Movement")
+	UE_DEPRECATED(5.7, "The ViewportInteraction module is being deprecated alongside VR Editor mode.")
+	UPROPERTY(config)
 	uint32 bScaleWorldWithDynamicPivot : 1;
 
 	/** When enabled, you can freely rotate and scale the world with two hands at the same time.  Otherwise, we'll detect whether to rotate or scale depending on how much of either gesture you initially perform. */
-	UPROPERTY(EditAnywhere, config, Category = "World Movement")
+	UE_DEPRECATED(5.7, "The ViewportInteraction module is being deprecated alongside VR Editor mode.")
+	UPROPERTY(config)
 	uint32 bAllowSimultaneousWorldScalingAndRotation : 1;
 };
 
diff --git a/Engine/Source/Editor/ViewportInteraction/Public/ViewportDragOperation.h b/Engine/Source/Editor/ViewportInteraction/Public/ViewportDragOperation.h
--- a/Engine/Source/Editor/ViewportInteraction/Public/ViewportDragOperation.h
+++ b/Engine/Source/Editor/ViewportInteraction/Public/ViewportDragOperation.h
@@ -22,11 +22,14 @@
 
 #define UE_API VIEWPORTINTERACTION_API
 
+PRAGMA_DISABLE_DEPRECATION_WARNINGS
+
+
 /**
  * Base class for interactable drag calculations
  */
 UCLASS(MinimalAPI)
-class UViewportDragOperation : public UObject
+class UE_DEPRECATED(5.7, "The ViewportInteraction module is being deprecated alongside VR Editor mode.") UViewportDragOperation : public UObject
 {
 	GENERATED_BODY()
 
@@ -57,7 +60,7 @@ public:
  * Container component for UViewportDragOperation that can be used by objects in the world that are draggable and implement the ViewportInteractableInterface
  */
 UCLASS(MinimalAPI)
-class UViewportDragOperationComponent : public UActorComponent
+class UE_DEPRECATED(5.7, "The ViewportInteraction module is being deprecated alongside VR Editor mode.") UViewportDragOperationComponent : public UActorComponent
 {
 	GENERATED_BODY()
 
@@ -96,7 +99,7 @@ private:
  * Data structure that holds all arguments that can be used while dragging a transformable.
  */
 USTRUCT()
-struct FDraggingTransformableData
+struct UE_DEPRECATED(5.7, "The ViewportInteraction module is being deprecated alongside VR Editor mode.") FDraggingTransformableData
 {
 	GENERATED_BODY()
 
@@ -157,4 +160,7 @@ struct FDraggingTransformableData
 	bool bAllowSnap;
 };
 
+
+PRAGMA_ENABLE_DEPRECATION_WARNINGS
+
 #undef UE_API
diff --git a/Engine/Source/Editor/ViewportInteraction/Public/ViewportInteractableInterface.h b/Engine/Source/Editor/ViewportInteraction/Public/ViewportInteractableInterface.h
--- a/Engine/Source/Editor/ViewportInteraction/Public/ViewportInteractableInterface.h
+++ b/Engine/Source/Editor/ViewportInteraction/Public/ViewportInteractableInterface.h
@@ -11,13 +11,21 @@
 
 #include "ViewportInteractableInterface.generated.h"
 
+
+PRAGMA_DISABLE_DEPRECATION_WARNINGS
+
+
 class UActorComponent;
 class UObject;
 class UViewportInteractor;
 struct FHitResult;
 
 UINTERFACE( MinimalAPI )
-class UViewportInteractableInterface : public UInterface
+class
+#if CPP
+	UE_DEPRECATED(5.7, "The ViewportInteraction module is being deprecated alongside VR Editor mode.")
+#endif
+	UViewportInteractableInterface : public UInterface
 {
 	GENERATED_UINTERFACE_BODY()
 };
@@ -25,7 +33,11 @@ class UViewportInteractableInterface : public UInterface
 /**
  * Interface for custom objects that can be interacted with by a interactor
  */
-class IViewportInteractableInterface
+class
+#if CPP
+	UE_DEPRECATED(5.7, "The ViewportInteraction module is being deprecated alongside VR Editor mode.")
+#endif
+	IViewportInteractableInterface
 {
 	GENERATED_IINTERFACE_BODY()
 
@@ -52,3 +64,6 @@ public:
 	/** Whether this interactable can be selected. */
 	virtual bool CanBeSelected() PURE_VIRTUAL( IViewportInteractableInterface::CanBeSelected, return false; );
 };
+
+
+PRAGMA_ENABLE_DEPRECATION_WARNINGS
diff --git a/Engine/Source/Editor/ViewportInteraction/Public/ViewportInteractionAssetContainer.h b/Engine/Source/Editor/ViewportInteraction/Public/ViewportInteractionAssetContainer.h
--- a/Engine/Source/Editor/ViewportInteraction/Public/ViewportInteractionAssetContainer.h
+++ b/Engine/Source/Editor/ViewportInteraction/Public/ViewportInteractionAssetContainer.h
@@ -14,7 +14,7 @@ class UStaticMesh;
  * Asset container for viewport interaction.
  */
 UCLASS(MinimalAPI)
-class UViewportInteractionAssetContainer : public UDataAsset
+class UE_DEPRECATED(5.7, "The ViewportInteraction module is being deprecated alongside VR Editor mode.") UViewportInteractionAssetContainer : public UDataAsset
 {
 	GENERATED_BODY()
 
diff --git a/Engine/Source/Editor/ViewportInteraction/Public/ViewportInteractionDragOperations.h b/Engine/Source/Editor/ViewportInteraction/Public/ViewportInteractionDragOperations.h
--- a/Engine/Source/Editor/ViewportInteraction/Public/ViewportInteractionDragOperations.h
+++ b/Engine/Source/Editor/ViewportInteraction/Public/ViewportInteractionDragOperations.h
@@ -10,11 +10,14 @@
 
 #define UE_API VIEWPORTINTERACTION_API
 
+PRAGMA_DISABLE_DEPRECATION_WARNINGS
+
+
 /**
  * Gizmo translation on one axis.
  */
 UCLASS(MinimalAPI)
-class UTranslationDragOperation : public UViewportDragOperation
+class UE_DEPRECATED(5.7, "The ViewportInteraction module is being deprecated alongside VR Editor mode.") UTranslationDragOperation : public UViewportDragOperation
 {
 	GENERATED_BODY()
 
@@ -26,7 +29,7 @@ public:
  * Gizmo translation on two axes.
  */
 UCLASS(MinimalAPI)
-class UPlaneTranslationDragOperation: public UViewportDragOperation
+class UE_DEPRECATED(5.7, "The ViewportInteraction module is being deprecated alongside VR Editor mode.") UPlaneTranslationDragOperation : public UViewportDragOperation
 {
 	GENERATED_BODY()
 
@@ -39,7 +42,7 @@ public:
  * Rotation around one axis based on input angle.
  */
 UCLASS(MinimalAPI)
-class URotateOnAngleDragOperation: public UViewportDragOperation
+class UE_DEPRECATED(5.7, "The ViewportInteraction module is being deprecated alongside VR Editor mode.") URotateOnAngleDragOperation : public UViewportDragOperation
 {
 	GENERATED_BODY()
 
@@ -67,7 +70,7 @@ private:
  * Scale on one axis.
  */
 UCLASS(MinimalAPI)
-class UScaleDragOperation : public UViewportDragOperation
+class UE_DEPRECATED(5.7, "The ViewportInteraction module is being deprecated alongside VR Editor mode.") UScaleDragOperation : public UViewportDragOperation
 {
 	GENERATED_BODY()
 
@@ -79,7 +82,7 @@ public:
  * Scale on all axes.
  */
 UCLASS(MinimalAPI)
-class  UUniformScaleDragOperation : public UViewportDragOperation
+class UE_DEPRECATED(5.7, "The ViewportInteraction module is being deprecated alongside VR Editor mode.") UUniformScaleDragOperation : public UViewportDragOperation
 {
 	GENERATED_BODY()
 
@@ -87,4 +90,7 @@ public:
 	UE_API virtual void ExecuteDrag(struct FDraggingTransformableData& DraggingData) override;
 };
 
+
+PRAGMA_ENABLE_DEPRECATION_WARNINGS
+
 #undef UE_API
diff --git a/Engine/Source/Editor/ViewportInteraction/Public/ViewportInteractionStyle.h b/Engine/Source/Editor/ViewportInteraction/Public/ViewportInteractionStyle.h
--- a/Engine/Source/Editor/ViewportInteraction/Public/ViewportInteractionStyle.h
+++ b/Engine/Source/Editor/ViewportInteraction/Public/ViewportInteractionStyle.h
@@ -12,7 +12,7 @@ class ISlateStyle;
 struct FSlateBrush;
 
 /**  */
-class FViewportInteractionStyle
+class UE_DEPRECATED(5.7, "The ViewportInteraction module is being deprecated alongside VR Editor mode.") FViewportInteractionStyle
 {
 public:
 
diff --git a/Engine/Source/Editor/ViewportInteraction/Public/ViewportInteractionTypes.h b/Engine/Source/Editor/ViewportInteraction/Public/ViewportInteractionTypes.h
--- a/Engine/Source/Editor/ViewportInteraction/Public/ViewportInteractionTypes.h
+++ b/Engine/Source/Editor/ViewportInteraction/Public/ViewportInteractionTypes.h
@@ -9,11 +9,14 @@
 
 #define UE_API VIEWPORTINTERACTION_API
 
+PRAGMA_DISABLE_DEPRECATION_WARNINGS
+
+
 /**
  * Represents a generic action
 */
 USTRUCT( BlueprintType )
-struct FViewportActionKeyInput
+struct UE_DEPRECATED(5.7, "The ViewportInteraction module is being deprecated alongside VR Editor mode.") FViewportActionKeyInput
 {
 	GENERATED_BODY()
 	
@@ -75,7 +78,7 @@ struct FViewportActionKeyInput
 
 /** Methods of dragging objects around in VR */
 UENUM()
-enum class EViewportInteractionDraggingMode : uint8
+enum class UE_DEPRECATED(5.7, "The ViewportInteraction module is being deprecated alongside VR Editor mode.") EViewportInteractionDraggingMode : uint8
 {
 	/** Not dragging right now with this hand */
 	Nothing,
@@ -103,7 +106,7 @@ enum class EViewportInteractionDraggingMode : uint8
 };
 
 /* Directions that a transform handle can face along any given axis */
-enum class ETransformGizmoHandleDirection
+enum class UE_DEPRECATED(5.7, "The ViewportInteraction module is being deprecated alongside VR Editor mode.") ETransformGizmoHandleDirection
 {
 	Negative,
 	Center,
@@ -113,7 +116,7 @@ enum class ETransformGizmoHandleDirection
 
 /** Placement of a handle in pivot space */
 USTRUCT()
-struct FTransformGizmoHandlePlacement
+struct UE_DEPRECATED(5.7, "The ViewportInteraction module is being deprecated alongside VR Editor mode.") FTransformGizmoHandlePlacement
 {
 	GENERATED_BODY()
 
@@ -130,11 +133,14 @@ struct FTransformGizmoHandlePlacement
 };
 
 
-enum class ELockedWorldDragMode
+enum class UE_DEPRECATED(5.7, "The ViewportInteraction module is being deprecated alongside VR Editor mode.") ELockedWorldDragMode
 {
 	Unlocked,
 	OnlyRotating,
 	OnlyScaling,
 };
 
+
+PRAGMA_ENABLE_DEPRECATION_WARNINGS
+
 #undef UE_API
diff --git a/Engine/Source/Editor/ViewportInteraction/Public/ViewportInteractor.h b/Engine/Source/Editor/ViewportInteraction/Public/ViewportInteractor.h
--- a/Engine/Source/Editor/ViewportInteraction/Public/ViewportInteractor.h
+++ b/Engine/Source/Editor/ViewportInteraction/Public/ViewportInteractor.h
@@ -14,11 +14,14 @@
 
 #define UE_API VIEWPORTINTERACTION_API
 
+PRAGMA_DISABLE_DEPRECATION_WARNINGS
+
+
 class AActor;
 
 /** Filter mode for GetHitResultFromLaserPointer*/
 UENUM(BlueprintType)
-enum class EHitResultGizmoFilterMode : uint8
+enum class UE_DEPRECATED(5.7, "The ViewportInteraction module is being deprecated alongside VR Editor mode.") EHitResultGizmoFilterMode : uint8
 {
 	All,
 	NoGizmos,
@@ -28,8 +31,8 @@ enum class EHitResultGizmoFilterMode : uint8
 /**
  * Represents the interactor in the world
  */
-UCLASS(MinimalAPI,  Abstract, Blueprintable, BlueprintType)
-class UViewportInteractor : public UObject
+UCLASS(MinimalAPI, Abstract, Blueprintable, BlueprintType)
+class UE_DEPRECATED(5.7, "The ViewportInteraction module is being deprecated alongside VR Editor mode.") UViewportInteractor : public UObject
 {
 	GENERATED_BODY()
 
@@ -328,4 +331,7 @@ private:
 	ViewportInteractionUtils::FOneEuroFilter SmoothingOneEuroFilter;
 };
 
+
+PRAGMA_ENABLE_DEPRECATION_WARNINGS
+
 #undef UE_API
diff --git a/Engine/Source/Editor/ViewportInteraction/Public/ViewportInteractorData.h b/Engine/Source/Editor/ViewportInteraction/Public/ViewportInteractorData.h
--- a/Engine/Source/Editor/ViewportInteraction/Public/ViewportInteractorData.h
+++ b/Engine/Source/Editor/ViewportInteraction/Public/ViewportInteractorData.h
@@ -8,8 +8,12 @@
 #include "Components/WidgetComponent.h"
 #include "ViewportInteractionTypes.h"
 
+
+PRAGMA_DISABLE_DEPRECATION_WARNINGS
+
+
 /** Represents a single virtual hand */
-struct FViewportInteractorData
+struct UE_DEPRECATED(5.7, "The ViewportInteraction module is being deprecated alongside VR Editor mode.") FViewportInteractorData
 {
 	//
 	// Positional data
@@ -201,3 +205,4 @@ struct FViewportInteractorData
 };
 
 
+PRAGMA_ENABLE_DEPRECATION_WARNINGS
diff --git a/Engine/Source/Editor/ViewportInteraction/Public/ViewportTransformable.h b/Engine/Source/Editor/ViewportInteraction/Public/ViewportTransformable.h
--- a/Engine/Source/Editor/ViewportInteraction/Public/ViewportTransformable.h
+++ b/Engine/Source/Editor/ViewportInteraction/Public/ViewportTransformable.h
@@ -12,7 +12,7 @@
 /**
  * Represents an object that we're actively interacting with, such as a selected actor
  */
-class FViewportTransformable
+class UE_DEPRECATED(5.7, "The ViewportInteraction module is being deprecated alongside VR Editor mode.") FViewportTransformable
 {
 
 public:
diff --git a/Engine/Source/Editor/ViewportInteraction/Public/ViewportTransformer.h b/Engine/Source/Editor/ViewportInteraction/Public/ViewportTransformer.h
--- a/Engine/Source/Editor/ViewportInteraction/Public/ViewportTransformer.h
+++ b/Engine/Source/Editor/ViewportInteraction/Public/ViewportTransformer.h
@@ -11,8 +11,11 @@
 
 #define UE_API VIEWPORTINTERACTION_API
 
+PRAGMA_DISABLE_DEPRECATION_WARNINGS
+
+
 UCLASS(MinimalAPI,  abstract )
-class UViewportTransformer : public UObject
+class UE_DEPRECATED(5.7, "The ViewportInteraction module is being deprecated alongside VR Editor mode.") UViewportTransformer : public UObject
 {
 	GENERATED_BODY()
 
@@ -54,4 +57,7 @@ protected:
 
 };
 
+
+PRAGMA_ENABLE_DEPRECATION_WARNINGS
+
 #undef UE_API
diff --git a/Engine/Source/Editor/ViewportInteraction/Public/ViewportWorldInteraction.h b/Engine/Source/Editor/ViewportInteraction/Public/ViewportWorldInteraction.h
--- a/Engine/Source/Editor/ViewportInteraction/Public/ViewportWorldInteraction.h
+++ b/Engine/Source/Editor/ViewportInteraction/Public/ViewportWorldInteraction.h
@@ -21,7 +21,10 @@
 
 #define UE_API VIEWPORTINTERACTION_API
 
-namespace ViewportWorldActionTypes
+PRAGMA_DISABLE_DEPRECATION_WARNINGS
+
+
+namespace UE_DEPRECATED(5.7, "The ViewportInteraction module is being deprecated alongside VR Editor mode.") ViewportWorldActionTypes
 {
 	static const FName NoAction( "NoAction" );
 	static const FName WorldMovement( "WorldMovement" );
@@ -39,14 +42,14 @@ class UViewportInteractionAssetContainer;
 class UViewportInteractor;
 
 UENUM()
-enum class EViewportWorldInteractionType : uint8
+enum class UE_DEPRECATED(5.7, "The ViewportInteraction module is being deprecated alongside VR Editor mode.") EViewportWorldInteractionType : uint8
 {
 	VR = 0,
 	Legacy = 1
 };
 
 UCLASS(MinimalAPI, BlueprintType, Transient)
-class UViewportWorldInteraction : public UEditorWorldExtension
+class UE_DEPRECATED(5.7, "The ViewportInteraction module is being deprecated alongside VR Editor mode.") UViewportWorldInteraction : public UEditorWorldExtension
 {
 	GENERATED_BODY()
 
@@ -764,4 +767,7 @@ private:
 	TSharedPtr<class FViewportInteractionInputProcessor> InputProcessor;
 };
 
+
+PRAGMA_ENABLE_DEPRECATION_WARNINGS
+
 #undef UE_API
