VPUtilities OSC Server - Allow specifying an override for the server address.

diff --git a/Engine/Plugins/Experimental/VirtualProductionUtilities/Source/VPUtilitiesEditor/Private/VPUtilitiesEditorModule.cpp b/Engine/Plugins/Experimental/VirtualProductionUtilities/Source/VPUtilitiesEditor/Private/VPUtilitiesEditorModule.cpp
--- a/Engine/Plugins/Experimental/VirtualProductionUtilities/Source/VPUtilitiesEditor/Private/VPUtilitiesEditorModule.cpp
+++ b/Engine/Plugins/Experimental/VirtualProductionUtilities/Source/VPUtilitiesEditor/Private/VPUtilitiesEditorModule.cpp
@@ -113,6 +113,14 @@ void FVPUtilitiesEditorModule::RegisterSettings()
 			LOCTEXT("VirtualProductionEditorSettingsDescription", "Configure the Virtual Production Editor settings."),
 			GetMutableDefault<UVPUtilitiesEditorSettings>());
 
+		ISettingsSectionPtr EditorUserSettingsSection = SettingsModule->RegisterSettings("Editor", "Plugins", "VirtualProduction",
+			LOCTEXT("VirtualProductionEditorUserSettingsName", "Virtual Production"),
+			LOCTEXT("VirtualProductionEditorUserSettingsNameDescription", "Configure the Virtual Production Editor user settings"),
+			GetMutableDefault<UVPUtilitiesEditorUserSettings>()
+		);
+
+		EditorUserSettingsSection->OnModified().BindRaw(this, &FVPUtilitiesEditorModule::OnSettingsModified);
+
 		EditorSettingsSection->OnModified().BindRaw(this, &FVPUtilitiesEditorModule::OnSettingsModified);
 	}
 
@@ -132,6 +140,7 @@ void FVPUtilitiesEditorModule::UnregisterSettings()
 	ISettingsModule* SettingsModule = FModuleManager::GetModulePtr<ISettingsModule>("Settings");
 	if (SettingsModule != nullptr)
 	{
+		SettingsModule->UnregisterSettings("Editor", "Plugins", "VirtualProduction");
 		SettingsModule->UnregisterSettings("Project", "Plugins", "VirtualProduction");
 		SettingsModule->UnregisterSettings("Project", "Plugins", "VirtualProductionEditor");
 	}
@@ -151,9 +160,25 @@ void FVPUtilitiesEditorModule::InitializeOSCServer()
 	}
 
 	const UVPUtilitiesEditorSettings* Settings = GetDefault<UVPUtilitiesEditorSettings>();
-	const FString& ServerAddress = Settings->OSCServerAddress;
+	FString ServerAddress = Settings->OSCServerAddress;
 	uint16 ServerPort = Settings->OSCServerPort;
 
+	const UVPUtilitiesEditorUserSettings* UserSettings = GetDefault<UVPUtilitiesEditorUserSettings>();
+	const FString& OverrideAddress = UserSettings->OverrideOSCServerAddress;
+	uint16 OverridePort = UserSettings->OverrideOSCServerPort;
+
+	if (!OverrideAddress.IsEmpty())
+	{
+		UE_LOG(LogVPUtilitiesEditor, Display, TEXT("OSC Server address was overridden from %s to %s."), *ServerAddress, *OverrideAddress);
+		ServerAddress = OverrideAddress;
+	}
+
+	if (OverridePort != 0)
+	{
+		UE_LOG(LogVPUtilitiesEditor, Display, TEXT("OSC Server port was overridden from %d to %d."), ServerPort, OverridePort);
+		ServerPort = OverridePort;
+	}
+
 	if (OSCServer)
 	{
 		OSCServer->SetAddress(ServerAddress, ServerPort);
diff --git a/Engine/Plugins/Experimental/VirtualProductionUtilities/Source/VPUtilitiesEditor/Public/VPUtilitiesEditorSettings.h b/Engine/Plugins/Experimental/VirtualProductionUtilities/Source/VPUtilitiesEditor/Public/VPUtilitiesEditorSettings.h
--- a/Engine/Plugins/Experimental/VirtualProductionUtilities/Source/VPUtilitiesEditor/Public/VPUtilitiesEditorSettings.h
+++ b/Engine/Plugins/Experimental/VirtualProductionUtilities/Source/VPUtilitiesEditor/Public/VPUtilitiesEditorSettings.h
@@ -7,6 +7,8 @@
 
 #include "EditorUtilityWidget.h"
 #include "EditorUtilityObject.h"
+#include "Misc/CommandLine.h"
+#include "Misc/Parse.h"
 
 #include "VPUtilitiesEditorSettings.generated.h"
 
@@ -95,4 +97,45 @@ public:
 	UE_DEPRECATED(5.5,"Code will be removed from UE5.7")
 	UPROPERTY(config, meta = (AllowedClasses = "/Script/Engine.StaticMesh"))
 	FSoftObjectPath VPSplinePreviewMeshPath;
+};
+
+/** Per-User editor settings for Virtual Production utilities. */
+UCLASS(MinimalAPI, config = EditorPerProjectUserSettings)
+class UVPUtilitiesEditorUserSettings : public UObject
+{
+	GENERATED_BODY()
+
+public:
+	virtual void PostInitProperties() override
+	{
+		Super::PostInitProperties();
+
+		FString OverrideAddress;
+		uint16 OverridePort;
+
+		if (FParse::Value(FCommandLine::Get(), TEXT("-VPOSCServerAddress="), OverrideAddress))
+		{
+			OverrideOSCServerAddress = MoveTemp(OverrideAddress);
+		}
+		
+		if (FParse::Value(FCommandLine::Get(), TEXT("-VPOSCServerPort="), OverridePort))
+		{
+			OverrideOSCServerPort = OverridePort;
+		}
+	}
+
+	/** 
+	 * If provided, this address will override the OSC server address provided in the project settings. 
+	 * Can also be specified on the command line through -VPOSCServerAddress=<Address>
+	 */
+	UPROPERTY(config, EditAnywhere, Category = "OSC", DisplayName = "OSC Server Address")
+	FString OverrideOSCServerAddress;
+
+	/**
+	 * If provided, this port will override the OSC server port provided in the project settings.
+	 * Can also be specified on the command line through -VPOSCServerPort=<Port>
+	 */
+	UPROPERTY(config, EditAnywhere, Category = "OSC", DisplayName = "OSC Server Port")
+	uint16 OverrideOSCServerPort;
+
 };
\ No newline at end of file
