Fix full screen widget for media output providers.

#rnx

diff --git a/Engine/Plugins/Experimental/VirtualProductionUtilities/Source/VPUtilities/Private/Widgets/Fullscreen/VPFullScreenUserWidget_PostProcess.cpp b/Engine/Plugins/Experimental/VirtualProductionUtilities/Source/VPUtilities/Private/Widgets/Fullscreen/VPFullScreenUserWidget_PostProcess.cpp
--- a/Engine/Plugins/Experimental/VirtualProductionUtilities/Source/VPUtilities/Private/Widgets/Fullscreen/VPFullScreenUserWidget_PostProcess.cpp
+++ b/Engine/Plugins/Experimental/VirtualProductionUtilities/Source/VPUtilities/Private/Widgets/Fullscreen/VPFullScreenUserWidget_PostProcess.cpp
@@ -39,11 +39,12 @@ void FVPFullScreenUserWidget_PostProcess::SetCustomPostProcessSettingsSource(TWe
 
 bool FVPFullScreenUserWidget_PostProcess::Display(UWorld* World, UUserWidget* Widget, TAttribute<float> InDPIScale)
 {
-	return CreateRenderer(World, Widget, MoveTemp(InDPIScale));
+	return CreateRenderer(World, Widget, MoveTemp(InDPIScale)) && InitPostProcessComponent(World);
 }
 
 void FVPFullScreenUserWidget_PostProcess::Hide(UWorld* World)
 {
+	ReleasePostProcessComponent();
 	FVPFullScreenUserWidget_PostProcessBase::Hide(World);
 }
 
